Added new test project WebTests;
Added http test cases;
Add the onlinetests web server created using django.

diff --git a/Engine/Source/Programs/WebTests/Private/TestHttp.cpp b/Engine/Source/Programs/WebTests/Private/TestHttp.cpp
new file mode 100644
--- /dev/null
+++ b/Engine/Source/Programs/WebTests/Private/TestHttp.cpp
@@ -0,0 +1,263 @@
+// Copyright Epic Games, Inc. All Rights Reserved.
+#include "CoreMinimal.h"
+#include "HAL/IConsoleManager.h"
+#include "HAL/PlatformProcess.h"
+#include "Online/HTTP/Public/HttpManager.h"
+#include "Online/HTTP/Public/Http.h"
+#include "Misc/CommandLine.h"
+#include "TestHarness.h"
+
+/**
+ *  HTTP Tests
+ *  -----------------------------------------------------------------------------------------------
+ *
+ *  PURPOSE:
+ *
+ *	Integration Tests to make sure all kinds of HTTP client features in C++ work well on different platforms,
+ *  including but not limited to error handing, retrying, threading, streaming, SSL and profiling.
+ *
+ *  Refer to WebTests/README.md for more info about how to run these tests
+ * 
+ *  -----------------------------------------------------------------------------------------------
+ */
+
+#define HTTP_TAG "[HTTP]"
+#define HTTP_TIME_DIFF_TOLERANCE 0.5f
+
+class FHttpModuleTestFixture
+{
+public:
+	FHttpModuleTestFixture()
+		: WebServerIp(TEXT("127.0.0.1"))
+	{
+		ParseSettingsFromCommandLine();
+
+		HttpModule = new FHttpModule();
+		IModuleInterface* Module = HttpModule;
+		Module->StartupModule();
+	}
+
+	virtual ~FHttpModuleTestFixture()
+	{
+		IModuleInterface* Module = HttpModule;
+		Module->ShutdownModule();
+		delete Module;
+	}
+
+	void ParseSettingsFromCommandLine()
+	{
+		FParse::Value(FCommandLine::Get(), TEXT("web_server_ip"), WebServerIp);
+	}
+
+	const FString UrlWithInvalidPortToTestConnectTimeout() const { return FString::Format(TEXT("http://{0}:{1}"), { *WebServerIp, 8765 }); }
+	const FString UrlBase() const { return FString::Format(TEXT("http://{0}:{1}"), { *WebServerIp, 8000 }); }
+	const FString UrlHttpTests() const { return FString::Format(TEXT("{0}/webtests/httptests"), { *UrlBase() }); }
+	const FString UrlToTestMethods() const { return FString::Format(TEXT("{0}/methods"), { *UrlHttpTests() }); }
+
+	FString WebServerIp;
+	FHttpModule* HttpModule;
+};
+
+TEST_CASE_METHOD(FHttpModuleTestFixture, "Shutdown http module without issue when there are ongoing http requests.", HTTP_TAG)
+{
+	uint32 ChunkSize = 1024 * 1024;
+	TArray<uint8> DataChunk;
+	DataChunk.SetNum(ChunkSize);
+	FMemory::Memset(DataChunk.GetData(), 'd', ChunkSize);
+
+	for (int32 i = 0; i < 10; ++i)
+	{
+		IHttpRequest* LeakingHttpRequest = FPlatformHttp::ConstructRequest(); // Leaking in purpose to make sure it's ok
+
+		TSharedRef<IHttpRequest, ESPMode::ThreadSafe> HttpRequest = HttpModule->CreateRequest();
+		HttpRequest->SetURL(UrlToTestMethods());
+		HttpRequest->SetVerb(TEXT("PUT"));
+		// TODO: Use some shared data, like cookie, openssl session etc.
+		HttpRequest->SetContent(DataChunk);
+		HttpRequest->OnProcessRequestComplete().BindLambda([](FHttpRequestPtr HttpRequest, FHttpResponsePtr HttpResponse, bool bSucceeded) {
+			CHECK(bSucceeded);
+		});
+		HttpRequest->ProcessRequest();
+	}
+
+	HttpModule->GetHttpManager().Tick(0.0f);
+}
+
+class FWaitUntilCompleteHttpFixture : public FHttpModuleTestFixture
+{
+public:
+	FWaitUntilCompleteHttpFixture()
+	{
+		HttpModule->GetHttpManager().SetRequestAddedDelegate(FHttpManagerRequestAddedDelegate::CreateRaw(this, &FWaitUntilCompleteHttpFixture::OnRequestAdded));
+		HttpModule->GetHttpManager().SetRequestCompletedDelegate(FHttpManagerRequestCompletedDelegate::CreateRaw(this, &FWaitUntilCompleteHttpFixture::OnRequestCompleted));
+	}
+
+	~FWaitUntilCompleteHttpFixture()
+	{
+		WaitUntilAllHttpRequestsComplete();
+
+		HttpModule->GetHttpManager().SetRequestAddedDelegate(FHttpManagerRequestAddedDelegate());
+		HttpModule->GetHttpManager().SetRequestCompletedDelegate(FHttpManagerRequestCompletedDelegate());
+	}
+
+	void OnRequestAdded(const FHttpRequestRef& Request)
+	{
+		++OngoingRequests;
+	}
+
+	void OnRequestCompleted(const FHttpRequestRef& Request)
+	{
+		if (ensure(OngoingRequests > 0))
+		{
+			--OngoingRequests;
+		}
+	}
+
+	void WaitUntilAllHttpRequestsComplete()
+	{
+		while (OngoingRequests != 0)
+		{
+			HttpModule->GetHttpManager().Tick(TickFrequency);
+			FPlatformProcess::Sleep(TickFrequency);
+		}
+	}
+
+	uint32 OngoingRequests = 0;
+	float TickFrequency = 1.0f / 60; /*60 FPS*/;
+};
+
+TEST_CASE_METHOD(FWaitUntilCompleteHttpFixture, "Http Methods - GET", HTTP_TAG)
+{
+	TSharedRef<IHttpRequest, ESPMode::ThreadSafe> HttpRequest = HttpModule->CreateRequest();
+	HttpRequest->SetURL(UrlToTestMethods());
+	HttpRequest->SetVerb(TEXT("GET"));
+	HttpRequest->OnProcessRequestComplete().BindLambda([](FHttpRequestPtr HttpRequest, FHttpResponsePtr HttpResponse, bool bSucceeded) {
+		CHECK(bSucceeded);
+		CHECK(HttpResponse->GetResponseCode() == 200);
+	});
+	HttpRequest->ProcessRequest();
+}
+
+TEST_CASE_METHOD(FWaitUntilCompleteHttpFixture, "Http Methods - POST", HTTP_TAG)
+{
+	TSharedRef<IHttpRequest, ESPMode::ThreadSafe> HttpRequest = HttpModule->CreateRequest();
+	HttpRequest->SetURL(UrlToTestMethods());
+	HttpRequest->SetVerb(TEXT("POST"));
+	HttpRequest->OnProcessRequestComplete().BindLambda([](FHttpRequestPtr HttpRequest, FHttpResponsePtr HttpResponse, bool bSucceeded) {
+		CHECK(bSucceeded);
+		CHECK(HttpResponse->GetResponseCode() == 200);
+	});
+	HttpRequest->ProcessRequest();
+}
+
+TEST_CASE_METHOD(FWaitUntilCompleteHttpFixture, "Http Methods - PUT", HTTP_TAG)
+{
+	TSharedRef<IHttpRequest, ESPMode::ThreadSafe> HttpRequest = HttpModule->CreateRequest();
+	HttpRequest->SetURL(UrlToTestMethods());
+	HttpRequest->SetVerb(TEXT("PUT"));
+	HttpRequest->OnProcessRequestComplete().BindLambda([](FHttpRequestPtr HttpRequest, FHttpResponsePtr HttpResponse, bool bSucceeded) {
+		CHECK(bSucceeded);
+		CHECK(HttpResponse->GetResponseCode() == 200);
+	});
+	HttpRequest->ProcessRequest();
+}
+
+TEST_CASE_METHOD(FWaitUntilCompleteHttpFixture, "Http Methods - DELETE", HTTP_TAG)
+{
+	TSharedRef<IHttpRequest, ESPMode::ThreadSafe> HttpRequest = HttpModule->CreateRequest();
+	HttpRequest->SetURL(UrlToTestMethods());
+	HttpRequest->SetVerb(TEXT("DELETE"));
+	HttpRequest->OnProcessRequestComplete().BindLambda([](FHttpRequestPtr HttpRequest, FHttpResponsePtr HttpResponse, bool bSucceeded) {
+		CHECK(bSucceeded);
+		CHECK(HttpResponse->GetResponseCode() == 200);
+	});
+	HttpRequest->ProcessRequest();
+}
+
+TEST_CASE_METHOD(FWaitUntilCompleteHttpFixture, "Http Methods - GET Large Response Content Without Chunks", HTTP_TAG)
+{
+	TSharedRef<IHttpRequest, ESPMode::ThreadSafe> HttpRequest = HttpModule->CreateRequest();
+	HttpRequest->SetURL(FString::Format(TEXT("{0}/get_large_response_without_chunks/{1}/"), { *UrlHttpTests(), 1024 * 1024/*bytes_number*/}));
+	HttpRequest->SetVerb(TEXT("GET"));
+	HttpRequest->OnProcessRequestComplete().BindLambda([](FHttpRequestPtr HttpRequest, FHttpResponsePtr HttpResponse, bool bSucceeded) {
+		CHECK(bSucceeded);
+		CHECK(HttpResponse->GetResponseCode() == 200);
+	});
+	HttpRequest->ProcessRequest();
+}
+
+TEST_CASE_METHOD(FWaitUntilCompleteHttpFixture, "Http request - connect timeout.", HTTP_TAG)
+{
+	TSharedRef<IHttpRequest, ESPMode::ThreadSafe> HttpRequest = HttpModule->CreateRequest();
+	HttpRequest->SetURL(UrlWithInvalidPortToTestConnectTimeout());
+	HttpRequest->SetVerb(TEXT("GET"));
+	HttpRequest->SetTimeout(7);
+	FDateTime StartTime = FDateTime::Now();
+	HttpRequest->OnProcessRequestComplete().BindLambda([StartTime](FHttpRequestPtr HttpRequest, FHttpResponsePtr HttpResponse, bool bSucceeded) {
+		CHECK(!bSucceeded);
+		CHECK(HttpResponse == nullptr);
+		CHECK(HttpRequest->GetStatus() == EHttpRequestStatus::Failed_ConnectionError);
+		FTimespan Timespan = FDateTime::Now() - StartTime;
+		float DurationInSeconds = Timespan.GetTotalSeconds();
+		CHECK(FMath::IsNearlyEqual(DurationInSeconds, 7, HTTP_TIME_DIFF_TOLERANCE));
+	});
+	HttpRequest->ProcessRequest();
+}
+
+TEST_CASE_METHOD(FWaitUntilCompleteHttpFixture, "Streaming http download - gold path.", HTTP_TAG)
+{
+	uint32 Chunks = 10;
+	uint32 ChunkSize = 1024*1024;
+
+	TSharedRef<IHttpRequest, ESPMode::ThreadSafe> HttpRequest = HttpModule->CreateRequest();
+	HttpRequest->SetURL(FString::Format(TEXT("{0}/streaming_download/{1}/{2}/"), { *UrlHttpTests(), Chunks, ChunkSize }));
+	HttpRequest->SetVerb(TEXT("GET"));
+	HttpRequest->OnProcessRequestComplete().BindLambda([Chunks, ChunkSize](FHttpRequestPtr HttpRequest, FHttpResponsePtr HttpResponse, bool bSucceeded) {
+		CHECK(HttpResponse->GetContentLength() == Chunks * ChunkSize);
+		CHECK(bSucceeded);
+		CHECK(HttpResponse->GetResponseCode() == 200);
+	});
+	HttpRequest->ProcessRequest();
+}
+
+TEST_CASE_METHOD(FWaitUntilCompleteHttpFixture, "Streaming http upload - gold path.", HTTP_TAG)
+{
+	TSharedRef<IHttpRequest, ESPMode::ThreadSafe> HttpRequest = HttpModule->CreateRequest();
+	HttpRequest->SetURL(FString::Format(TEXT("{0}/streaming_upload"), { *UrlHttpTests() }));
+	HttpRequest->SetVerb(TEXT("POST"));
+
+	const char* BoundaryLabel = "test_http_boundary";
+	HttpRequest->SetHeader(TEXT("Content-Type"), FString::Format(TEXT("multipart/form-data; boundary={0}"), { BoundaryLabel }));
+
+	// Not really reading file here in order to simplify the test flow. It will be sent by chunks in http request, 
+	const uint32 FileSize = 10*1024*1024;
+	char* FileData = (char*)FMemory::Malloc(FileSize + 1);
+	FMemory::Memset(FileData, 'd', FileSize);
+	FileData[FileSize + 1] = '\0';
+
+	TArray<uint8> ContentData;
+	const int32 ContentMaxSize = FileSize + 256/*max length of format string*/;
+	ContentData.Reserve(ContentMaxSize);
+	const int32 ContentLength = FCStringAnsi::Snprintf(
+		(char*)ContentData.GetData(),
+		ContentMaxSize,
+		"--%s\r\n"
+		"Content-Disposition: form-data; name=\"file\"; filename=\"bigfile.zip\"\r\n"
+		"Content-Type: application/octet-stream\r\n\r\n"
+		"%s\r\n"
+		"--%s--",
+		BoundaryLabel, FileData, BoundaryLabel);
+
+	FMemory::Free(FileData);
+
+	CHECK(ContentLength > 0);
+	CHECK(ContentLength < ContentMaxSize);
+	ContentData.SetNumUninitialized(ContentLength);
+	HttpRequest->SetContent(MoveTemp(ContentData));
+
+	HttpRequest->OnProcessRequestComplete().BindLambda([](FHttpRequestPtr HttpRequest, FHttpResponsePtr HttpResponse, bool bSucceeded) {
+		CHECK(bSucceeded);
+		CHECK(HttpResponse->GetResponseCode() == 200);
+	});
+	HttpRequest->ProcessRequest();
+}
diff --git a/Engine/Source/Programs/WebTests/README.md b/Engine/Source/Programs/WebTests/README.md
new file mode 100644
--- /dev/null
+++ b/Engine/Source/Programs/WebTests/README.md
@@ -0,0 +1,24 @@
+WebTests
+--------------------
+
+This is a collection of tests for web features like http, websocket, ssl etc.
+
+# Running the Tests
+
+## Run the WebServer(NOTE: This WebServer is created for TEST PURPOSE ONLY, it's not production ready!!!):
+	In ./WebServer folder:
+		install python3.x
+		> pip install -r requirements.txt
+		> runserver.bat
+	Or if you use docker:
+		> dockerbuildandrun.bat
+
+## Run the tests from VS:
+	Set `WebTests` as the startup project and set Solution Configuration to a `Development`.
+	If running tests on other devices, pass in the ip as command line args, after extra args AT THE END, like: "-- --web_server_ip=your.pc.ip.address"
+	Compile and debug
+
+# Adding new test case in WebServer:
+	Add/change the code in ./WebServer code, most likely in httptests/urls.py and httptests/views.py, and save. Code will be reloaded if the web server is running
+	When using docker, also need to run dockerbuildandrun.bat
+	For more info about how to code in django, check https://docs.djangoproject.com/en/
diff --git a/Engine/Source/Programs/WebTests/WebServer/dockerbuildandrun.bat b/Engine/Source/Programs/WebTests/WebServer/dockerbuildandrun.bat
new file mode 100755
--- /dev/null
+++ b/Engine/Source/Programs/WebTests/WebServer/dockerbuildandrun.bat
@@ -0,0 +1,14 @@
+@echo off
+
+:: Build image
+docker build . -t webtests
+
+set CONTAINER="webtestscontainer"
+:: Check if container exist and delete
+docker ps -a | findstr /c:"%CONTAINER%" >nul 2>&1
+if %errorlevel% equ 0 (
+	docker rm -f %CONTAINER%
+)
+
+:: Run container with image
+docker run --name %CONTAINER% -d -p 8000:8000 webtests
diff --git a/Engine/Source/Programs/WebTests/WebServer/dockerfile b/Engine/Source/Programs/WebTests/WebServer/dockerfile
new file mode 100644
--- /dev/null
+++ b/Engine/Source/Programs/WebTests/WebServer/dockerfile
@@ -0,0 +1,11 @@
+FROM python:3.10.4
+ENV DockerHOME=[REDACTED_USER_HOME]/webtests
+RUN mkdir -p $DockerHOME
+WORKDIR $DockerHOME
+ENV PYTHONDONTWRITEBYTECODE 1
+ENV PYTHONUNBUFFERED 1
+RUN pip install --upgrade pip
+COPY . $DockerHOME
+RUN pip install -r requirements.txt
+EXPOSE 8000
+CMD python manage.py runserver 0.0.0.0:8000
diff --git a/Engine/Source/Programs/WebTests/WebServer/httptests/admin.py b/Engine/Source/Programs/WebTests/WebServer/httptests/admin.py
new file mode 100644
--- /dev/null
+++ b/Engine/Source/Programs/WebTests/WebServer/httptests/admin.py
@@ -0,0 +1,3 @@
+from django.contrib import admin
+
+# Register your models here.
diff --git a/Engine/Source/Programs/WebTests/WebServer/httptests/apps.py b/Engine/Source/Programs/WebTests/WebServer/httptests/apps.py
new file mode 100644
--- /dev/null
+++ b/Engine/Source/Programs/WebTests/WebServer/httptests/apps.py
@@ -0,0 +1,6 @@
+from django.apps import AppConfig
+
+
+class HttpTestsConfig(AppConfig):
+    default_auto_field = 'django.db.models.BigAutoField'
+    name = 'httptests'
diff --git a/Engine/Source/Programs/WebTests/WebServer/httptests/models.py b/Engine/Source/Programs/WebTests/WebServer/httptests/models.py
new file mode 100644
--- /dev/null
+++ b/Engine/Source/Programs/WebTests/WebServer/httptests/models.py
@@ -0,0 +1,3 @@
+from django.db import models
+
+# Create your models here.
diff --git a/Engine/Source/Programs/WebTests/WebServer/httptests/tests.py b/Engine/Source/Programs/WebTests/WebServer/httptests/tests.py
new file mode 100644
--- /dev/null
+++ b/Engine/Source/Programs/WebTests/WebServer/httptests/tests.py
@@ -0,0 +1,3 @@
+from django.test import TestCase
+
+# Create your tests here.
diff --git a/Engine/Source/Programs/WebTests/WebServer/httptests/urls.py b/Engine/Source/Programs/WebTests/WebServer/httptests/urls.py
new file mode 100644
--- /dev/null
+++ b/Engine/Source/Programs/WebTests/WebServer/httptests/urls.py
@@ -0,0 +1,11 @@
+from django.urls import path
+
+from . import views
+
+urlpatterns = [
+    path("httptests/methods", views.methods, name="methods"),
+    path("httptests/get_large_response_without_chunks/<int:bytes_number>/", views.get_large_response_without_chunks, name="get_large_response_without_chunks"),
+    path("httptests/nonstreaming_receivetimeout/<int:wait_time>/", views.nonstreaming_receivetimeout, name="nonstreaming_receivetimeout"),
+    path("httptests/streaming_download/<int:chunks>/<int:chunk_size>/", views.streaming_download, name="streaming_download"),
+    path("httptests/streaming_upload", views.streaming_upload, name="streaming_upload"),
+]
diff --git a/Engine/Source/Programs/WebTests/WebServer/httptests/views.py b/Engine/Source/Programs/WebTests/WebServer/httptests/views.py
new file mode 100644
--- /dev/null
+++ b/Engine/Source/Programs/WebTests/WebServer/httptests/views.py
@@ -0,0 +1,47 @@
+from django.shortcuts import render
+from django.http.response import JsonResponse
+from django.http import JsonResponse
+from django.http import StreamingHttpResponse
+from rest_framework.decorators import api_view
+import time
+import logging
+
+logger = logging.getLogger('django')
+
+@api_view(['GET', 'POST', 'DELETE', 'PUT'])
+def methods(request):
+    return JsonResponse({})
+
+@api_view(['GET'])
+def get_large_response_without_chunks(request, bytes_number):
+    data = "d" * bytes_number
+    return JsonResponse({'data' : data})
+
+@api_view(['GET'])
+def nonstreaming_receivetimeout(request, wait_time):
+    time.sleep(wait_time)
+    return JsonResponse({})
+
+@api_view(['GET'])
+def streaming_download(request, chunks, chunk_size):
+    def data_chunk_generator():
+        for x in range(chunks):
+            data_chunk = 'd' * chunk_size
+            yield data_chunk
+    return StreamingHttpResponse(data_chunk_generator())
+
+@api_view(['POST'])
+def streaming_upload(request):
+    f = request.FILES["file"]
+    if f.multiple_chunks():
+        # big file
+        total_chunks = 0
+        total_bytes = 0
+        for chunk in f.chunks():
+            total_chunks += 1
+            total_bytes += len(chunk)
+        logger.info("Received %d chunks: %d bytes in total", total_chunks, total_bytes)
+    else:
+        # small file
+        logger.info("Received the whole file: %d bytes", f.size)
+    return JsonResponse({})
diff --git a/Engine/Source/Programs/WebTests/WebServer/manage.py b/Engine/Source/Programs/WebTests/WebServer/manage.py
new file mode 100644
--- /dev/null
+++ b/Engine/Source/Programs/WebTests/WebServer/manage.py
@@ -0,0 +1,22 @@
+#!/usr/bin/env python
+"""Django's command-line utility for administrative tasks."""
+import os
+import sys
+
+
+def main():
+    """Run administrative tasks."""
+    os.environ.setdefault('DJANGO_SETTINGS_MODULE', 'webtests.settings')
+    try:
+        from django.core.management import execute_from_command_line
+    except ImportError as exc:
+        raise ImportError(
+            "Couldn't import Django. Are you sure it's installed and "
+            "available on your PYTHONPATH environment variable? Did you "
+            "forget to activate a virtual environment?"
+        ) from exc
+    execute_from_command_line(sys.argv)
+
+
+if __name__ == '__main__':
+    main()
diff --git a/Engine/Source/Programs/WebTests/WebServer/requirements.txt b/Engine/Source/Programs/WebTests/WebServer/requirements.txt
new file mode 100644
--- /dev/null
+++ b/Engine/Source/Programs/WebTests/WebServer/requirements.txt
@@ -0,0 +1,2 @@
+django==4.1.7
+djangorestframework==3.14.0
diff --git a/Engine/Source/Programs/WebTests/WebServer/runserver.bat b/Engine/Source/Programs/WebTests/WebServer/runserver.bat
new file mode 100755
--- /dev/null
+++ b/Engine/Source/Programs/WebTests/WebServer/runserver.bat
@@ -0,0 +1,2 @@
+@echo off
+python manage.py runserver 0.0.0.0:8000
diff --git a/Engine/Source/Programs/WebTests/WebServer/webtests/asgi.py b/Engine/Source/Programs/WebTests/WebServer/webtests/asgi.py
new file mode 100644
--- /dev/null
+++ b/Engine/Source/Programs/WebTests/WebServer/webtests/asgi.py
@@ -0,0 +1,16 @@
+"""
+ASGI config for webtests project.
+
+It exposes the ASGI callable as a module-level variable named ``application``.
+
+For more information on this file, see
+https://docs.djangoproject.com/en/4.1/howto/deployment/asgi/
+"""
+
+import os
+
+from django.core.asgi import get_asgi_application
+
+os.environ.setdefault('DJANGO_SETTINGS_MODULE', 'webtests.settings')
+
+application = get_asgi_application()
diff --git a/Engine/Source/Programs/WebTests/WebServer/webtests/settings.py b/Engine/Source/Programs/WebTests/WebServer/webtests/settings.py
new file mode 100644
--- /dev/null
+++ b/Engine/Source/Programs/WebTests/WebServer/webtests/settings.py
@@ -0,0 +1,130 @@
+"""
+Django settings for webtests project.
+
+Generated by 'django-admin startproject' using Django 4.1.7.
+
+For more information on this file, see
+https://docs.djangoproject.com/en/4.1/topics/settings/
+
+For the full list of settings and their values, see
+https://docs.djangoproject.com/en/4.1/ref/settings/
+"""
+
+from pathlib import Path
+
+# Build paths inside the project like this: BASE_DIR / 'subdir'.
+BASE_DIR = Path(__file__).resolve().parent.parent
+
+
+# Quick-start development settings - unsuitable for production
+# See https://docs.djangoproject.com/en/4.1/howto/deployment/checklist/
+
+# SECURITY WARNING: keep the secret key used in production secret!
+SECRET_KEY = 'django-insecure-x=3k^30px$8#ouewp0(1er-gpm@90k^z5$6bxx(qtot92yy$%w'
+
+# SECURITY WARNING: don't run with debug turned on in production!
+DEBUG = True
+
+# Allow any host, because this web server is only for local test purpose
+ALLOWED_HOSTS = ['*']
+
+
+# Application definition
+
+INSTALLED_APPS = [
+    'django.contrib.admin',
+    'django.contrib.auth',
+    'django.contrib.contenttypes',
+    'django.contrib.sessions',
+    'django.contrib.messages',
+    'django.contrib.staticfiles',
+]
+
+MIDDLEWARE = [
+    'django.middleware.security.SecurityMiddleware',
+    'django.contrib.sessions.middleware.SessionMiddleware',
+    'django.middleware.common.CommonMiddleware',
+    'django.middleware.csrf.CsrfViewMiddleware',
+    'django.contrib.auth.middleware.AuthenticationMiddleware',
+    'django.contrib.messages.middleware.MessageMiddleware',
+    'django.middleware.clickjacking.XFrameOptionsMiddleware',
+]
+
+ROOT_URLCONF = 'webtests.urls'
+
+TEMPLATES = [
+    {
+        'BACKEND': 'django.template.backends.django.DjangoTemplates',
+        'DIRS': [],
+        'APP_DIRS': True,
+        'OPTIONS': {
+            'context_processors': [
+                'django.template.context_processors.debug',
+                'django.template.context_processors.request',
+                'django.contrib.auth.context_processors.auth',
+                'django.contrib.messages.context_processors.messages',
+            ],
+        },
+    },
+]
+
+WSGI_APPLICATION = 'webtests.wsgi.application'
+
+
+# Database
+# https://docs.djangoproject.com/en/4.1/ref/settings/#databases
+
+DATABASES = {
+    'default': {
+        'ENGINE': 'django.db.backends.sqlite3',
+        'NAME': BASE_DIR / 'db.sqlite3',
+    }
+}
+
+
+# Password validation
+# https://docs.djangoproject.com/en/4.1/ref/settings/#auth-password-validators
+
+AUTH_PASSWORD_VALIDATORS = [
+    {
+        'NAME': 'django.contrib.auth.password_validation.UserAttributeSimilarityValidator',
+    },
+    {
+        'NAME': 'django.contrib.auth.password_validation.MinimumLengthValidator',
+    },
+    {
+        'NAME': 'django.contrib.auth.password_validation.CommonPasswordValidator',
+    },
+    {
+        'NAME': 'django.contrib.auth.password_validation.NumericPasswordValidator',
+    },
+]
+
+
+# Internationalization
+# https://docs.djangoproject.com/en/4.1/topics/i18n/
+
+LANGUAGE_CODE = 'en-us'
+
+TIME_ZONE = 'UTC'
+
+USE_I18N = True
+
+USE_TZ = True
+
+
+# Static files (CSS, JavaScript, Images)
+# https://docs.djangoproject.com/en/4.1/howto/static-files/
+
+STATIC_URL = 'static/'
+
+# Default primary key field type
+# https://docs.djangoproject.com/en/4.1/ref/settings/#default-auto-field
+
+DEFAULT_AUTO_FIELD = 'django.db.models.BigAutoField'
+
+LOGGING = {
+    "version": 1,  # the dictConfig format version
+    "disable_existing_loggers": False,  # retain the default loggers
+    "level": "DEBUG",
+}
diff --git a/Engine/Source/Programs/WebTests/WebServer/webtests/urls.py b/Engine/Source/Programs/WebTests/WebServer/webtests/urls.py
new file mode 100644
--- /dev/null
+++ b/Engine/Source/Programs/WebTests/WebServer/webtests/urls.py
@@ -0,0 +1,22 @@
+"""webtests URL Configuration
+
+The `urlpatterns` list routes URLs to views. For more information please see:
+    https://docs.djangoproject.com/en/4.1/topics/http/urls/
+Examples:
+Function views
+    1. Add an import:  from my_app import views
+    2. Add a URL to urlpatterns:  path('', views.home, name='home')
+Class-based views
+    1. Add an import:  from other_app.views import Home
+    2. Add a URL to urlpatterns:  path('', Home.as_view(), name='home')
+Including another URLconf
+    1. Import the include() function: from django.urls import include, path
+    2. Add a URL to urlpatterns:  path('blog/', include('blog.urls'))
+"""
+from django.contrib import admin
+from django.urls import include, path
+
+urlpatterns = [
+    path("webtests/", include("httptests.urls")),
+    path('admin/', admin.site.urls),
+]
diff --git a/Engine/Source/Programs/WebTests/WebServer/webtests/wsgi.py b/Engine/Source/Programs/WebTests/WebServer/webtests/wsgi.py
new file mode 100644
--- /dev/null
+++ b/Engine/Source/Programs/WebTests/WebServer/webtests/wsgi.py
@@ -0,0 +1,16 @@
+"""
+WSGI config for webtests project.
+
+It exposes the WSGI callable as a module-level variable named ``application``.
+
+For more information on this file, see
+https://docs.djangoproject.com/en/4.1/howto/deployment/wsgi/
+"""
+
+import os
+
+from django.core.wsgi import get_wsgi_application
+
+os.environ.setdefault('DJANGO_SETTINGS_MODULE', 'webtests.settings')
+
+application = get_wsgi_application()
diff --git a/Engine/Source/Programs/WebTests/WebTests.Build.cs b/Engine/Source/Programs/WebTests/WebTests.Build.cs
new file mode 100644
--- /dev/null
+++ b/Engine/Source/Programs/WebTests/WebTests.Build.cs
@@ -0,0 +1,18 @@
+// Copyright Epic Games, Inc. All Rights Reserved.
+
+using UnrealBuildTool;
+
+public class WebTests : TestModuleRules
+{
+	public WebTests(ReadOnlyTargetRules Target) : base(Target)
+	{
+		PrivateDependencyModuleNames.AddRange(
+			new string[] {
+				"ApplicationCore",
+				"Core",
+				"HTTP",
+				"Chaos",
+				"ChaosUserDataPT"
+			});
+	}
+}
diff --git a/Engine/Source/Programs/WebTests/WebTests.Target.cs b/Engine/Source/Programs/WebTests/WebTests.Target.cs
new file mode 100644
--- /dev/null
+++ b/Engine/Source/Programs/WebTests/WebTests.Target.cs
@@ -0,0 +1,14 @@
+// Copyright Epic Games, Inc. All Rights Reserved.
+
+using UnrealBuildTool;
+
+[SupportedPlatforms(UnrealPlatformClass.All)]
+public class WebTestsTarget : TestTargetRules
+{
+	public WebTestsTarget(TargetInfo Target) : base(Target)
+	{
+		IncludeOrderVersion = EngineIncludeOrderVersion.Latest;
+		bCompileAgainstCoreUObject = true;
+		bCompileAgainstApplicationCore = true;
+	}
+}
diff --git a/Engine/Source/Runtime/Core/Private/Misc/URLRequestFilter.cpp b/Engine/Source/Runtime/Core/Private/Misc/URLRequestFilter.cpp
--- a/Engine/Source/Runtime/Core/Private/Misc/URLRequestFilter.cpp
+++ b/Engine/Source/Runtime/Core/Private/Misc/URLRequestFilter.cpp
@@ -33,6 +33,11 @@ FURLRequestFilter::FURLRequestFilter(FRequestMap&& InAllowedRequests)
 
 void FURLRequestFilter::UpdateConfig(const TCHAR* ConfigSectionRootName, const FString& ConfigFileName)
 {
+	if (ConfigFileName.IsEmpty() || GConfig == nullptr)
+	{
+		return;
+	}
+
 	AllowedRequests.Empty();
 
 	FConfigFile* File = GConfig->FindConfigFile(ConfigFileName);
