Updated the WebTests doc since UE5 integrated with new tests param;
Use sections for http methods tests.

diff --git a/Engine/Source/Programs/WebTests/Private/TestHttp.cpp b/Engine/Source/Programs/WebTests/Private/TestHttp.cpp
--- a/Engine/Source/Programs/WebTests/Private/TestHttp.cpp
+++ b/Engine/Source/Programs/WebTests/Private/TestHttp.cpp
@@ -3,8 +3,8 @@
 #include "HAL/FileManager.h"
 #include "HAL/IConsoleManager.h"
 #include "HAL/PlatformProcess.h"
-#include "Online/HTTP/Public/HttpManager.h"
-#include "Online/HTTP/Public/Http.h"
+#include "HttpManager.h"
+#include "Http.h"
 #include "Misc/CommandLine.h"
 #include "TestHarness.h"
 
@@ -127,47 +127,28 @@ public:
 	float TickFrequency = 1.0f / 60; /*60 FPS*/;
 };
 
-TEST_CASE_METHOD(FWaitUntilCompleteHttpFixture, "Http Methods - GET", HTTP_TAG)
+TEST_CASE_METHOD(FWaitUntilCompleteHttpFixture, "Http Methods", HTTP_TAG)
 {
 	TSharedRef<IHttpRequest, ESPMode::ThreadSafe> HttpRequest = HttpModule->CreateRequest();
 	HttpRequest->SetURL(UrlToTestMethods());
-	HttpRequest->SetVerb(TEXT("GET"));
-	HttpRequest->OnProcessRequestComplete().BindLambda([](FHttpRequestPtr HttpRequest, FHttpResponsePtr HttpResponse, bool bSucceeded) {
-		CHECK(bSucceeded);
-		CHECK(HttpResponse->GetResponseCode() == 200);
-	});
-	HttpRequest->ProcessRequest();
-}
-
-TEST_CASE_METHOD(FWaitUntilCompleteHttpFixture, "Http Methods - POST", HTTP_TAG)
-{
-	TSharedRef<IHttpRequest, ESPMode::ThreadSafe> HttpRequest = HttpModule->CreateRequest();
-	HttpRequest->SetURL(UrlToTestMethods());
-	HttpRequest->SetVerb(TEXT("POST"));
-	HttpRequest->OnProcessRequestComplete().BindLambda([](FHttpRequestPtr HttpRequest, FHttpResponsePtr HttpResponse, bool bSucceeded) {
-		CHECK(bSucceeded);
-		CHECK(HttpResponse->GetResponseCode() == 200);
-	});
-	HttpRequest->ProcessRequest();
-}
 
-TEST_CASE_METHOD(FWaitUntilCompleteHttpFixture, "Http Methods - PUT", HTTP_TAG)
-{
-	TSharedRef<IHttpRequest, ESPMode::ThreadSafe> HttpRequest = HttpModule->CreateRequest();
-	HttpRequest->SetURL(UrlToTestMethods());
-	HttpRequest->SetVerb(TEXT("PUT"));
-	HttpRequest->OnProcessRequestComplete().BindLambda([](FHttpRequestPtr HttpRequest, FHttpResponsePtr HttpResponse, bool bSucceeded) {
-		CHECK(bSucceeded);
-		CHECK(HttpResponse->GetResponseCode() == 200);
-	});
-	HttpRequest->ProcessRequest();
-}
+	SECTION("GET")
+	{
+		HttpRequest->SetVerb(TEXT("GET"));
+	}
+	SECTION("POST")
+	{
+		HttpRequest->SetVerb(TEXT("POST"));
+	}
+	SECTION("PUT")
+	{
+		HttpRequest->SetVerb(TEXT("PUT"));
+	}
+	SECTION("DELETE")
+	{
+		HttpRequest->SetVerb(TEXT("DELETE"));
+	}
 
-TEST_CASE_METHOD(FWaitUntilCompleteHttpFixture, "Http Methods - DELETE", HTTP_TAG)
-{
-	TSharedRef<IHttpRequest, ESPMode::ThreadSafe> HttpRequest = HttpModule->CreateRequest();
-	HttpRequest->SetURL(UrlToTestMethods());
-	HttpRequest->SetVerb(TEXT("DELETE"));
 	HttpRequest->OnProcessRequestComplete().BindLambda([](FHttpRequestPtr HttpRequest, FHttpResponsePtr HttpResponse, bool bSucceeded) {
 		CHECK(bSucceeded);
 		CHECK(HttpResponse->GetResponseCode() == 200);
@@ -175,7 +156,7 @@ TEST_CASE_METHOD(FWaitUntilCompleteHttpFixture, "Http Methods - DELETE", HTTP_TA
 	HttpRequest->ProcessRequest();
 }
 
-TEST_CASE_METHOD(FWaitUntilCompleteHttpFixture, "Http Methods - GET Large Response Content Without Chunks", HTTP_TAG)
+TEST_CASE_METHOD(FWaitUntilCompleteHttpFixture, "Get large response content without chunks", HTTP_TAG)
 {
 	TSharedRef<IHttpRequest, ESPMode::ThreadSafe> HttpRequest = HttpModule->CreateRequest();
 	HttpRequest->SetURL(FString::Format(TEXT("{0}/get_large_response_without_chunks/{1}/"), { *UrlHttpTests(), 1024 * 1024/*bytes_number*/}));
@@ -187,7 +168,7 @@ TEST_CASE_METHOD(FWaitUntilCompleteHttpFixture, "Http Methods - GET Large Respon
 	HttpRequest->ProcessRequest();
 }
 
-TEST_CASE_METHOD(FWaitUntilCompleteHttpFixture, "Http request - connect timeout.", HTTP_TAG)
+TEST_CASE_METHOD(FWaitUntilCompleteHttpFixture, "Http request connect timeout", HTTP_TAG)
 {
 	TSharedRef<IHttpRequest, ESPMode::ThreadSafe> HttpRequest = HttpModule->CreateRequest();
 	HttpRequest->SetURL(UrlWithInvalidPortToTestConnectTimeout());
diff --git a/Engine/Source/Programs/WebTests/README.md b/Engine/Source/Programs/WebTests/README.md
--- a/Engine/Source/Programs/WebTests/README.md
+++ b/Engine/Source/Programs/WebTests/README.md
@@ -15,7 +15,7 @@ This is a collection of tests for web features like http, websocket, ssl etc.
 
 ## Run the tests from VS:
 	Set `WebTests` as the startup project and set Solution Configuration to a `Development`.
-	If running tests on other devices, pass in the ip as command line args, after extra args AT THE END, like: "-- --web_server_ip=your.pc.ip.address"
+	If running tests on other devices, pass in the ip as command line args, after extra args AT THE END, like: "--extra-args --web_server_ip=your.pc.ip.address"
 	Compile and debug
 
 # Adding new test case in WebServer:
