Add non-game thread support for http request.

diff --git a/Engine/Source/Programs/WebTests/Private/TestHttp.cpp b/Engine/Source/Programs/WebTests/Private/TestHttp.cpp
--- a/Engine/Source/Programs/WebTests/Private/TestHttp.cpp
+++ b/Engine/Source/Programs/WebTests/Private/TestHttp.cpp
@@ -3,6 +3,8 @@
 #include "HAL/FileManager.h"
 #include "HAL/IConsoleManager.h"
 #include "HAL/PlatformProcess.h"
+#include "HAL/Runnable.h"
+#include "HAL/RunnableThread.h"
 #include "HttpManager.h"
 #include "Http.h"
 #include "Misc/CommandLine.h"
@@ -116,7 +118,7 @@ public:
 
 	void WaitUntilAllHttpRequestsComplete()
 	{
-		while (OngoingRequests != 0)
+		while (!bRunningThreadRequest && OngoingRequests != 0)
 		{
 			HttpModule->GetHttpManager().Tick(TickFrequency);
 			FPlatformProcess::Sleep(TickFrequency);
@@ -125,6 +127,7 @@ public:
 
 	uint32 OngoingRequests = 0;
 	float TickFrequency = 1.0f / 60; /*60 FPS*/;
+	bool bRunningThreadRequest = false;
 };
 
 TEST_CASE_METHOD(FWaitUntilCompleteHttpFixture, "Http Methods", HTTP_TAG)
@@ -365,21 +368,140 @@ TEST_CASE_METHOD(FWaitUntilCompleteHttpFixture, "Streaming http upload - gold pa
 	HttpRequest->ProcessRequest();
 }
 
-TEST_CASE_METHOD(FWaitUntilCompleteHttpFixture, "Pre check http request", HTTP_TAG)
+class FWaitThreadedHttpFixture : public FWaitUntilCompleteHttpFixture, public FRunnable
+{
+public:
+	DECLARE_DELEGATE(FRunActualTestCodeDelegate);
+
+	FWaitThreadedHttpFixture()
+	{
+		bRunningThreadRequest = true;
+	}
+
+	// FRunnable interface
+	virtual uint32 Run() override
+	{
+		ThreadCallback.ExecuteIfBound();
+		bRunningThreadRequest = false;
+		return 0;
+	}
+
+	void StartTestHttpThread()
+	{
+		RunnableThread = TSharedPtr<FRunnableThread>(FRunnableThread::Create(this, TEXT("Test Http Thread")));
+	}
+
+	FRunActualTestCodeDelegate ThreadCallback;
+	TSharedPtr<FRunnableThread> RunnableThread;
+};
+
+TEST_CASE_METHOD(FWaitThreadedHttpFixture, "Http streaming download request can work in non game thread", HTTP_TAG)
+{
+	ThreadCallback.BindLambda([this]() {
+		TSharedRef<IHttpRequest> HttpRequest = HttpModule->CreateRequest();
+		HttpRequest->SetURL(FString::Format(TEXT("{0}/streaming_download/{1}/{2}/"), { *UrlHttpTests(), 3/*chunks*/, 1024/*chunk_size*/}));
+		HttpRequest->SetVerb(TEXT("GET"));
+		HttpRequest->SetDelegateThreadPolicy(EHttpRequestDelegateThreadPolicy::CompleteOnHttpThread);
+
+		class FTestHttpReceiveStream final : public FArchive
+		{
+		public:
+			virtual void Serialize(void* V, int64 Length) override
+			{
+				// No matter what's the thread policy, Serialize always get called in http thread.
+				CHECK(!IsInGameThread());
+			}
+		};
+		CHECK(HttpRequest->SetResponseBodyReceiveStream(MakeShared<FTestHttpReceiveStream>()));
+
+		HttpRequest->OnProcessRequestComplete().BindLambda([](FHttpRequestPtr HttpRequest, FHttpResponsePtr HttpResponse, bool bSucceeded) {
+			// EHttpRequestDelegateThreadPolicy::CompleteOnHttpThread was used, so not in game thread here
+			CHECK(!IsInGameThread());
+			CHECK(bSucceeded);
+			CHECK(HttpResponse->GetResponseCode() == 200);
+		});
+
+		HttpRequest->ProcessRequest();
+	});
+
+	StartTestHttpThread();
+}
+
+namespace UE
+{
+namespace TestHttp
+{
+
+void SetupURLRequestFilter(FHttpModule* HttpModule)
 {
 	// Pre check will fail when domain is not allowed
 	UE::Core::FURLRequestFilter::FRequestMap SchemeMap;
 	SchemeMap.Add(TEXT("http"), TArray<FString>{TEXT("epicgames.com")});
 	UE::Core::FURLRequestFilter Filter{SchemeMap};
 	HttpModule->GetHttpManager().SetURLRequestFilter(Filter);
+}
+
+}
+}
+
+TEST_CASE_METHOD(FWaitUntilCompleteHttpFixture, "Http request pre check will fail", HTTP_TAG)
+{
+	// Pre check will fail when domain is not allowed
+	UE::TestHttp::SetupURLRequestFilter(HttpModule);
 
 	TSharedRef<IHttpRequest, ESPMode::ThreadSafe> HttpRequest = HttpModule->CreateRequest();
 	HttpRequest->SetVerb(TEXT("GET"));
-	HttpRequest->SetURL(UrlToTestMethods()); // Not in allowed domains: epicgames.com
-	HttpRequest->OnProcessRequestComplete().BindLambda([](FHttpRequestPtr HttpRequest, FHttpResponsePtr HttpResponse, bool bSucceeded) {
-		CHECK(!bSucceeded);
-	});
+	HttpRequest->SetURL(UrlToTestMethods());
+
+	SECTION("on game thread")
+	{
+		HttpRequest->OnProcessRequestComplete().BindLambda([](FHttpRequestPtr HttpRequest, FHttpResponsePtr HttpResponse, bool bSucceeded) {
+			CHECK(IsInGameThread());
+			CHECK(!bSucceeded);
+		});
+	}
+	SECTION("on http thread")
+	{
+		HttpRequest->SetDelegateThreadPolicy(EHttpRequestDelegateThreadPolicy::CompleteOnHttpThread);
+		HttpRequest->OnProcessRequestComplete().BindLambda([](FHttpRequestPtr HttpRequest, FHttpResponsePtr HttpResponse, bool bSucceeded) {
+			CHECK(!IsInGameThread());
+			CHECK(!bSucceeded);
+		});
+	}
 
 	HttpRequest->ProcessRequest();
 }
 
+TEST_CASE_METHOD(FWaitThreadedHttpFixture, "Threaded http request pre check will fail", HTTP_TAG)
+{
+	ThreadCallback.BindLambda([this]() {
+		// Pre check will fail when domain is not allowed
+		UE::TestHttp::SetupURLRequestFilter(HttpModule);
+
+		TSharedRef<IHttpRequest, ESPMode::ThreadSafe> HttpRequest = HttpModule->CreateRequest();
+		HttpRequest->SetVerb(TEXT("GET"));
+		HttpRequest->SetURL(UrlToTestMethods());
+
+		SECTION("on game thread")
+		{
+			HttpRequest->OnProcessRequestComplete().BindLambda([](FHttpRequestPtr HttpRequest, FHttpResponsePtr HttpResponse, bool bSucceeded) {
+				CHECK(IsInGameThread());
+				CHECK(!bSucceeded);
+			});
+		}
+		SECTION("on http thread")
+		{
+			HttpRequest->SetDelegateThreadPolicy(EHttpRequestDelegateThreadPolicy::CompleteOnHttpThread);
+			HttpRequest->OnProcessRequestComplete().BindLambda([](FHttpRequestPtr HttpRequest, FHttpResponsePtr HttpResponse, bool bSucceeded) {
+				CHECK(!IsInGameThread());
+				CHECK(!bSucceeded);
+			});
+		}
+
+		HttpRequest->ProcessRequest();
+	});
+
+	StartTestHttpThread();
+}
+
+// TODO: Add cancel test, with multiple cancel calls
diff --git a/Engine/Source/Runtime/Online/BuildPatchServices/Private/Tests/Mock/HttpRequest.mock.h b/Engine/Source/Runtime/Online/BuildPatchServices/Private/Tests/Mock/HttpRequest.mock.h
--- a/Engine/Source/Runtime/Online/BuildPatchServices/Private/Tests/Mock/HttpRequest.mock.h
+++ b/Engine/Source/Runtime/Online/BuildPatchServices/Private/Tests/Mock/HttpRequest.mock.h
@@ -189,6 +189,17 @@ namespace BuildPatchServices
 			return float();
 		}
 
+		virtual void SetDelegateThreadPolicy(EHttpRequestDelegateThreadPolicy InThreadPolicy) override
+		{
+			MOCK_FUNC_NOT_IMPLEMENTED("FMockHttpRequest::SetDelegateThreadPolicy");
+		}
+
+		virtual EHttpRequestDelegateThreadPolicy GetDelegateThreadPolicy() const override
+		{
+			MOCK_FUNC_NOT_IMPLEMENTED("FMockHttpRequest::GetDelegateThreadPolicy");
+			return EHttpRequestDelegateThreadPolicy::CompleteOnGameThread;
+		}
+
 	public:
 		FHttpRequestProgressDelegate HttpRequestProgressDelegate;
 		FHttpRequestCompleteDelegate HttpRequestCompleteDelegate;
diff --git a/Engine/Source/Runtime/Online/HTTP/Private/Curl/CurlHttp.cpp b/Engine/Source/Runtime/Online/HTTP/Private/Curl/CurlHttp.cpp
--- a/Engine/Source/Runtime/Online/HTTP/Private/Curl/CurlHttp.cpp
+++ b/Engine/Source/Runtime/Online/HTTP/Private/Curl/CurlHttp.cpp
@@ -1006,22 +1006,7 @@ bool FCurlHttpRequest::ProcessRequest()
 
 	if (!PreCheck() || !SetupRequest())
 	{
-		if (!IsInGameThread())
-		{
-			// Always finish on the game thread
-			FHttpModule::Get().GetHttpManager().AddGameThreadTask([StrongThis = StaticCastSharedRef<FCurlHttpRequest>(AsShared())]()
-			{
-				StrongThis->FinishedRequest();
-			});
-			return true;
-		}
-		else
-		{
-			// Cleanup and call delegate
-			FinishedRequest();
-		}
-
-		return false;
+		return FinishRequestNotInHttpManager();
 	}
 
 	// Clear the info cache log so we don't output messages from previous requests when reusing/retrying a request
@@ -1055,11 +1040,6 @@ bool FCurlHttpRequest::StartThreadedRequest()
 	return true;
 }
 
-void FCurlHttpRequest::FinishRequest()
-{
-	FinishedRequest();
-}
-
 bool FCurlHttpRequest::IsThreadedRequestComplete()
 {
 	if (bCanceled)
@@ -1115,18 +1095,9 @@ void FCurlHttpRequest::CancelRequest()
 	{
 		HttpManager.CancelThreadedRequest(SharedThis(this));
 	}
-	else if (!IsInGameThread())
-	{
-		// Always finish on the game thread
-		FHttpModule::Get().GetHttpManager().AddGameThreadTask([StrongThis = StaticCastSharedRef<FCurlHttpRequest>(AsShared())]()
-		{
-			StrongThis->FinishedRequest();
-		});
-	}
 	else
 	{
-		// Finish immediately
-		FinishedRequest();
+		FinishRequestNotInHttpManager();
 	}
 }
 
@@ -1165,7 +1136,7 @@ void FCurlHttpRequest::CheckProgressDelegate()
 
 void FCurlHttpRequest::BroadcastNewlyReceivedHeaders()
 {
-	check(IsInGameThread());
+	check(IsInGameThread() || DelegateThreadPolicy == EHttpRequestDelegateThreadPolicy::CompleteOnHttpThread);
 	if (Response.IsValid())
 	{
 		// Process the headers received on the HTTP thread and merge them into the response's list of headers and then broadcast the new headers
@@ -1192,10 +1163,11 @@ void FCurlHttpRequest::BroadcastNewlyReceivedHeaders()
 	}
 }
 
-void FCurlHttpRequest::FinishedRequest()
+void FCurlHttpRequest::FinishRequest()
 {
-	check(IsInGameThread());
-	QUICK_SCOPE_CYCLE_COUNTER(STAT_FCurlHttpRequest_FinishedRequest);
+	check(IsInGameThread() || DelegateThreadPolicy == EHttpRequestDelegateThreadPolicy::CompleteOnHttpThread);
+
+	QUICK_SCOPE_CYCLE_COUNTER(STAT_FCurlHttpRequest_FinishRequest);
 
 	curl_easy_setopt(EasyHandle, CURLOPT_SHARE, nullptr);
 	
diff --git a/Engine/Source/Runtime/Online/HTTP/Private/Curl/CurlHttp.h b/Engine/Source/Runtime/Online/HTTP/Private/Curl/CurlHttp.h
--- a/Engine/Source/Runtime/Online/HTTP/Private/Curl/CurlHttp.h
+++ b/Engine/Source/Runtime/Online/HTTP/Private/Curl/CurlHttp.h
@@ -353,12 +353,6 @@ private:
 	 */
 	bool SetupRequest();
 
-	/**
-	 * Process state for a finished request that no longer needs to be ticked
-	 * Calls the completion delegate
-	 */
-	void FinishedRequest();
-
 	/**
 	 * Trigger the request progress delegate if progress has changed
 	 */
diff --git a/Engine/Source/Runtime/Online/HTTP/Private/GenericPlatform/GenericPlatformHttp.cpp b/Engine/Source/Runtime/Online/HTTP/Private/GenericPlatform/GenericPlatformHttp.cpp
--- a/Engine/Source/Runtime/Online/HTTP/Private/GenericPlatform/GenericPlatformHttp.cpp
+++ b/Engine/Source/Runtime/Online/HTTP/Private/GenericPlatform/GenericPlatformHttp.cpp
@@ -107,6 +107,8 @@ public:
 	virtual const FHttpResponsePtr GetResponse() const override { return nullptr; }
 	virtual void Tick(float DeltaSeconds) override {}
 	virtual float GetElapsedTime() const override { return 0.0f; }
+	virtual void SetDelegateThreadPolicy(EHttpRequestDelegateThreadPolicy InDelegateThreadPolicy) override {}
+	virtual EHttpRequestDelegateThreadPolicy GetDelegateThreadPolicy() const override { return EHttpRequestDelegateThreadPolicy::CompleteOnGameThread; }
 };
 
 void FGenericPlatformHttp::Init()
diff --git a/Engine/Source/Runtime/Online/HTTP/Private/GenericPlatform/HttpRequestCommon.cpp b/Engine/Source/Runtime/Online/HTTP/Private/GenericPlatform/HttpRequestCommon.cpp
--- a/Engine/Source/Runtime/Online/HTTP/Private/GenericPlatform/HttpRequestCommon.cpp
+++ b/Engine/Source/Runtime/Online/HTTP/Private/GenericPlatform/HttpRequestCommon.cpp
@@ -42,3 +42,13 @@ bool FHttpRequestCommon::PreCheck() const
 	return true;
 }
 
+void FHttpRequestCommon::SetDelegateThreadPolicy(EHttpRequestDelegateThreadPolicy InDelegateThreadPolicy)
+{ 
+	DelegateThreadPolicy = InDelegateThreadPolicy; 
+}
+
+EHttpRequestDelegateThreadPolicy FHttpRequestCommon::GetDelegateThreadPolicy() const
+{ 
+	return DelegateThreadPolicy; 
+}
+
diff --git a/Engine/Source/Runtime/Online/HTTP/Private/HttpManager.cpp b/Engine/Source/Runtime/Online/HTTP/Private/HttpManager.cpp
--- a/Engine/Source/Runtime/Online/HTTP/Private/HttpManager.cpp
+++ b/Engine/Source/Runtime/Online/HTTP/Private/HttpManager.cpp
@@ -281,6 +281,12 @@ void FHttpManager::AddGameThreadTask(TFunction<void()>&& Task)
 	}
 }
 
+void FHttpManager::AddHttpThreadTask(TFunction<void()>&& Task)
+{
+	check(Thread);
+	Thread->AddHttpThreadTask(MoveTemp(Task));
+}
+
 FHttpThreadBase* FHttpManager::CreateHttpThread()
 {
 	return new FLegacyHttpThread();
@@ -423,7 +429,10 @@ bool FHttpManager::Tick(float DeltaSeconds)
 		{
 			FHttpRequestRef CompletedRequestRef = CompletedRequest->AsShared();
 			Requests.Remove(CompletedRequestRef);
-			CompletedRequest->FinishRequest();
+			if (CompletedRequest->GetDelegateThreadPolicy() == EHttpRequestDelegateThreadPolicy::CompleteOnGameThread)
+			{
+				CompletedRequest->FinishRequest();
+			}
 			BroadcastHttpRequestCompleted(CompletedRequestRef);
 		}
 	}
diff --git a/Engine/Source/Runtime/Online/HTTP/Private/HttpThread.cpp b/Engine/Source/Runtime/Online/HTTP/Private/HttpThread.cpp
--- a/Engine/Source/Runtime/Online/HTTP/Private/HttpThread.cpp
+++ b/Engine/Source/Runtime/Online/HTTP/Private/HttpThread.cpp
@@ -109,6 +109,13 @@ uint32 FHttpThreadBase::Run()
 
 void FHttpThreadBase::Tick()
 {
+	// Run HttpThread tasks
+	TFunction<void()> Task = nullptr;
+	while (HttpThreadQueue.Dequeue(Task))
+	{
+		check(Task);
+		Task();
+	}
 }
 
 bool FHttpThreadBase::NeedsSingleThreadTick() const
@@ -132,6 +139,14 @@ void FHttpThreadBase::UpdateConfigs()
 	}
 }
 
+void FHttpThreadBase::AddHttpThreadTask(TFunction<void()>&& Task)
+{
+	if (Task)
+	{
+		HttpThreadQueue.Enqueue(MoveTemp(Task));
+	}
+}
+
 void FHttpThreadBase::HttpThreadTick(float DeltaSeconds)
 {
 	// empty
@@ -277,6 +292,12 @@ void FHttpThreadBase::Process(TArray<IHttpThreadedRequest*>& RequestsToCancel, T
 			SCOPE_CYCLE_COUNTER(STAT_HTTPThread_CompleteThreadedRequest);
 
 			CompleteThreadedRequest(Request);
+
+			if (Request->GetDelegateThreadPolicy() == EHttpRequestDelegateThreadPolicy::CompleteOnHttpThread)
+			{
+				Request->FinishRequest();
+			}
+
 			CompletedThreadedRequests.Enqueue(Request);
 		}
 		RequestsToComplete.Reset();
diff --git a/Engine/Source/Runtime/Online/HTTP/Private/HttpThread.h b/Engine/Source/Runtime/Online/HTTP/Private/HttpThread.h
--- a/Engine/Source/Runtime/Online/HTTP/Private/HttpThread.h
+++ b/Engine/Source/Runtime/Online/HTTP/Private/HttpThread.h
@@ -78,6 +78,13 @@ public:
 	 */
 	virtual void UpdateConfigs();
 
+	/**
+	 * Add task to be ran on the http thread next tick
+	 *
+	 * @param Task The task to be ran next tick
+	 */
+	void AddHttpThreadTask(TFunction<void()>&& Task);
+
 protected:
 
 	/**
@@ -163,6 +170,9 @@ protected:
 	 * Added to on HTTP thread, processed then cleared on game thread (Single producer, single consumer)
 	 */
 	TSpscQueue<IHttpThreadedRequest*> CompletedThreadedRequests;
+
+	/** Queue of tasks to run on the game thread */
+	TQueue<TFunction<void()>, EQueueMode::Mpsc> HttpThreadQueue;
 };
 
 class FLegacyHttpThread	: public FHttpThreadBase
diff --git a/Engine/Source/Runtime/Online/HTTP/Private/IHttpThreadedRequest.cpp b/Engine/Source/Runtime/Online/HTTP/Private/IHttpThreadedRequest.cpp
new file mode 100644
--- /dev/null
+++ b/Engine/Source/Runtime/Online/HTTP/Private/IHttpThreadedRequest.cpp
@@ -0,0 +1,42 @@
+// Copyright Epic Games, Inc. All Rights Reserved.
+
+#include "IHttpThreadedRequest.h"
+#include "HttpModule.h"
+#include "HttpManager.h"
+
+bool IHttpThreadedRequest::FinishRequestNotInHttpManager()
+{
+	if (IsInGameThread())
+	{
+		if (DelegateThreadPolicy == EHttpRequestDelegateThreadPolicy::CompleteOnGameThread)
+		{
+			FinishRequest();
+			return false;
+		}
+		else
+		{
+			FHttpModule::Get().GetHttpManager().AddHttpThreadTask([StrongThis = StaticCastSharedRef<IHttpThreadedRequest>(AsShared())]()
+				{
+					StrongThis->FinishRequest();
+				});
+			return true;
+		}
+	}
+	else
+	{
+		if (DelegateThreadPolicy == EHttpRequestDelegateThreadPolicy::CompleteOnHttpThread)
+		{
+			FinishRequest();
+			return false;
+		}
+		else
+		{
+			FHttpModule::Get().GetHttpManager().AddGameThreadTask([StrongThis = StaticCastSharedRef<IHttpThreadedRequest>(AsShared())]()
+				{
+					StrongThis->FinishRequest();
+				});
+			return true;
+		}
+	}
+}
+
diff --git a/Engine/Source/Runtime/Online/HTTP/Private/IHttpThreadedRequest.h b/Engine/Source/Runtime/Online/HTTP/Private/IHttpThreadedRequest.h
--- a/Engine/Source/Runtime/Online/HTTP/Private/IHttpThreadedRequest.h
+++ b/Engine/Source/Runtime/Online/HTTP/Private/IHttpThreadedRequest.h
@@ -17,4 +17,10 @@ public:
 	virtual void FinishRequest() = 0;
 
 protected:
+	/**
+	 * Finish the request when it's not in http manager
+	 * 
+	 * @return false if the request finished immediate without waiting, otherwise it needs to be finished in another thread instead of current thread
+	 */
+	bool FinishRequestNotInHttpManager();
 };
diff --git a/Engine/Source/Runtime/Online/HTTP/Private/NullHttp.cpp b/Engine/Source/Runtime/Online/HTTP/Private/NullHttp.cpp
--- a/Engine/Source/Runtime/Online/HTTP/Private/NullHttp.cpp
+++ b/Engine/Source/Runtime/Online/HTTP/Private/NullHttp.cpp
@@ -195,6 +195,15 @@ void FNullHttpRequest::FinishedRequest()
 	OnProcessRequestComplete().ExecuteIfBound(Request, NULL, false);
 }
 
+void FNullHttpRequest::SetDelegateThreadPolicy(EHttpRequestDelegateThreadPolicy InThreadPolicy)
+{
+}
+
+EHttpRequestDelegateThreadPolicy FNullHttpRequest::GetDelegateThreadPolicy() const
+{
+	return EHttpRequestDelegateThreadPolicy::CompleteOnGameThread;
+}
+
 // FNullHttpResponse
 
 FString FNullHttpResponse::GetURL() const
diff --git a/Engine/Source/Runtime/Online/HTTP/Private/NullHttp.h b/Engine/Source/Runtime/Online/HTTP/Private/NullHttp.h
--- a/Engine/Source/Runtime/Online/HTTP/Private/NullHttp.h
+++ b/Engine/Source/Runtime/Online/HTTP/Private/NullHttp.h
@@ -39,6 +39,8 @@ public:
 	virtual const FHttpResponsePtr GetResponse() const override;
 	virtual void Tick(float DeltaSeconds) override;
 	virtual float GetElapsedTime() const override;
+	virtual void SetDelegateThreadPolicy(EHttpRequestDelegateThreadPolicy InThreadPolicy) override;
+	virtual EHttpRequestDelegateThreadPolicy GetDelegateThreadPolicy() const override;
 
 	FNullHttpRequest()
 		: CompletionStatus(EHttpRequestStatus::NotStarted)
diff --git a/Engine/Source/Runtime/Online/HTTP/Public/GenericPlatform/HttpRequestCommon.h b/Engine/Source/Runtime/Online/HTTP/Public/GenericPlatform/HttpRequestCommon.h
--- a/Engine/Source/Runtime/Online/HTTP/Public/GenericPlatform/HttpRequestCommon.h
+++ b/Engine/Source/Runtime/Online/HTTP/Public/GenericPlatform/HttpRequestCommon.h
@@ -12,6 +12,8 @@ class HTTP_API FHttpRequestCommon : public FHttpRequestImpl
 public:
 	// IHttpRequest
 	virtual EHttpRequestStatus::Type GetStatus() const override;
+	virtual void SetDelegateThreadPolicy(EHttpRequestDelegateThreadPolicy InDelegateThreadPolicy) override;
+	virtual EHttpRequestDelegateThreadPolicy GetDelegateThreadPolicy() const override;
 
 protected:
 	/**
@@ -22,4 +24,7 @@ protected:
 protected:
 	/** Current status of request being processed */
 	EHttpRequestStatus::Type CompletionStatus = EHttpRequestStatus::NotStarted;
+
+	/** Thread policy about which thread to complete this request */
+	EHttpRequestDelegateThreadPolicy DelegateThreadPolicy = EHttpRequestDelegateThreadPolicy::CompleteOnGameThread;
 };
diff --git a/Engine/Source/Runtime/Online/HTTP/Public/HttpManager.h b/Engine/Source/Runtime/Online/HTTP/Public/HttpManager.h
--- a/Engine/Source/Runtime/Online/HTTP/Public/HttpManager.h
+++ b/Engine/Source/Runtime/Online/HTTP/Public/HttpManager.h
@@ -218,6 +218,13 @@ public:
 	 */
 	void AddGameThreadTask(TFunction<void()>&& Task);
 
+	/**
+	 * Add task to be ran on the http thread next tick
+	 *
+	 * @param Task The task to be ran next tick
+	 */
+	void AddHttpThreadTask(TFunction<void()>&& Task);
+
 	/**
 	 * Set url request filter through code, instead of setting it through config.
 	 *
diff --git a/Engine/Source/Runtime/Online/HTTP/Public/HttpRequestAdapter.h b/Engine/Source/Runtime/Online/HTTP/Public/HttpRequestAdapter.h
--- a/Engine/Source/Runtime/Online/HTTP/Public/HttpRequestAdapter.h
+++ b/Engine/Source/Runtime/Online/HTTP/Public/HttpRequestAdapter.h
@@ -43,6 +43,8 @@ public:
 	virtual float                         GetElapsedTime() const override                                          { return HttpRequest->GetElapsedTime(); }
 	virtual EHttpRequestStatus::Type	  GetStatus() const override                                               { return HttpRequest->GetStatus(); }
 	virtual void                          Tick(float DeltaSeconds) override                                        { HttpRequest->Tick(DeltaSeconds); }
+	virtual void                          SetDelegateThreadPolicy(EHttpRequestDelegateThreadPolicy InThreadPolicy) override{ HttpRequest->SetDelegateThreadPolicy(InThreadPolicy); }
+	virtual EHttpRequestDelegateThreadPolicy GetDelegateThreadPolicy() const override                                 { return HttpRequest->GetDelegateThreadPolicy(); }
 
 protected:
     TSharedRef<IHttpRequest, ESPMode::ThreadSafe> HttpRequest;
diff --git a/Engine/Source/Runtime/Online/HTTP/Public/Interfaces/IHttpRequest.h b/Engine/Source/Runtime/Online/HTTP/Public/Interfaces/IHttpRequest.h
--- a/Engine/Source/Runtime/Online/HTTP/Public/Interfaces/IHttpRequest.h
+++ b/Engine/Source/Runtime/Online/HTTP/Public/Interfaces/IHttpRequest.h
@@ -62,6 +62,15 @@ namespace EHttpRequestStatus
 	}
 }
 
+/**
+ * Enumerates thread policy about which thread to complete the http request
+ */
+enum class EHttpRequestDelegateThreadPolicy : uint8
+{
+	CompleteOnGameThread = 0,
+	CompleteOnHttpThread,
+};
+
 class IHttpRequest;
 class IHttpResponse;
 
@@ -209,6 +218,7 @@ public:
 
 	/**
 	 * Sets the content of the request to stream directly from an archive.
+	 * NOTE: The Stream->Serialize will be called from another thread other than the game thread
 	 *
 	 * @param Stream - archive from which the payload should be streamed.
 	 * @return True if the archive can be used to stream the request. False otherwise.
@@ -217,6 +227,7 @@ public:
 
 	/**
 	 * Sets the content of the request to stream directly from an delegate.
+	 * NOTE: The delegate will be called from another thread other than the game thread
 	 *
 	 * @param StreamDelegate - delegate from which the payload should be streamed.
 	 * @return True if the delegate can be used to stream the request. False otherwise.
@@ -360,6 +371,18 @@ public:
 	 */
 	virtual float GetElapsedTime() const = 0;
 
+	/**
+	 * Set thread policy about which thread to complete this request
+	 */
+	virtual void SetDelegateThreadPolicy(EHttpRequestDelegateThreadPolicy InThreadPolicy) = 0;
+
+	/**
+	 * Get thread policy about which thread to complete this request
+	 * 
+	 * @return The thread policy
+	 */
+	virtual EHttpRequestDelegateThreadPolicy GetDelegateThreadPolicy() const = 0;
+
 	/** 
 	 * Destructor for overrides 
 	 */
