Add a test case to verify http request can be reused;
Add a test case to verify http request can run in parallel.

diff --git a/Engine/Source/Programs/WebTests/Private/TestHttp.cpp b/Engine/Source/Programs/WebTests/Private/TestHttp.cpp
--- a/Engine/Source/Programs/WebTests/Private/TestHttp.cpp
+++ b/Engine/Source/Programs/WebTests/Private/TestHttp.cpp
@@ -86,6 +86,57 @@ TEST_CASE_METHOD(FHttpModuleTestFixture, "Shutdown http module without issue whe
 	HttpModule->GetHttpManager().Tick(0.0f);
 }
 
+class FWaitUntilQuitFromTestFixture : public FHttpModuleTestFixture
+{
+public:
+	FWaitUntilQuitFromTestFixture()
+	{
+	}
+
+	~FWaitUntilQuitFromTestFixture()
+	{
+		WaitUntilAllHttpRequestsComplete();
+	}
+
+	void WaitUntilAllHttpRequestsComplete()
+	{
+		while (!bQuitRequested)
+		{
+			HttpModule->GetHttpManager().Tick(TickFrequency);
+			FPlatformProcess::Sleep(TickFrequency);
+		}
+	}
+
+	float TickFrequency = 1.0f / 60; /*60 FPS*/;
+	bool bQuitRequested = false;
+};
+
+TEST_CASE_METHOD(FWaitUntilQuitFromTestFixture, "Http request can be reused", HTTP_TAG)
+{
+	TSharedRef<IHttpRequest> HttpRequest = HttpModule->CreateRequest();
+	HttpRequest->SetURL(UrlToTestMethods());
+	HttpRequest->SetVerb(TEXT("POST"));
+
+	HttpRequest->OnProcessRequestComplete().BindLambda([this](FHttpRequestPtr HttpRequest, FHttpResponsePtr HttpResponse, bool bSucceeded) {
+		CHECK(bSucceeded);
+		CHECK(HttpResponse->GetResponseCode() == 200);
+
+		uint32 Chunks = 3;
+		uint32 ChunkSize = 1024;
+		HttpRequest->SetURL(FString::Format(TEXT("{0}/streaming_download/{1}/{2}/"), { *UrlHttpTests(), Chunks, ChunkSize }));
+		HttpRequest->SetVerb(TEXT("GET"));
+		HttpRequest->OnProcessRequestComplete().BindLambda([this, Chunks, ChunkSize](FHttpRequestPtr HttpRequest, FHttpResponsePtr HttpResponse, bool bSucceeded) {
+			CHECK(bSucceeded);
+			REQUIRE(HttpResponse != nullptr);
+			CHECK(HttpResponse->GetResponseCode() == 200);
+			CHECK(HttpResponse->GetContentLength() == Chunks * ChunkSize);
+			bQuitRequested = true;
+		});
+		HttpRequest->ProcessRequest();
+	});
+	HttpRequest->ProcessRequest();
+}
+
 class FWaitUntilCompleteHttpFixture : public FHttpModuleTestFixture
 {
 public:
@@ -193,7 +244,7 @@ TEST_CASE_METHOD(FWaitUntilCompleteHttpFixture, "Http request connect timeout",
 
 TEST_CASE_METHOD(FWaitUntilCompleteHttpFixture, "Streaming http download", HTTP_TAG)
 {
-	uint32 Chunks = 10;
+	uint32 Chunks = 3;
 	uint32 ChunkSize = 1024*1024;
 
 	TSharedRef<IHttpRequest, ESPMode::ThreadSafe> HttpRequest = HttpModule->CreateRequest();
@@ -326,6 +377,25 @@ TEST_CASE_METHOD(FWaitUntilCompleteHttpFixture, "Streaming http download", HTTP_
 	HttpRequest->ProcessRequest();
 }
 
+TEST_CASE_METHOD(FWaitUntilCompleteHttpFixture, "Can run parallel stream download requests", HTTP_TAG)
+{
+	uint32 Chunks = 5;
+	uint32 ChunkSize = 1024*1024;
+
+	for (int i = 0; i < 3; ++i)
+	{
+		TSharedRef<IHttpRequest, ESPMode::ThreadSafe> HttpRequest = HttpModule->CreateRequest();
+		HttpRequest->SetURL(FString::Format(TEXT("{0}/streaming_download/{1}/{2}/"), { *UrlHttpTests(), Chunks, ChunkSize }));
+		HttpRequest->SetVerb(TEXT("GET"));
+		HttpRequest->OnProcessRequestComplete().BindLambda([Chunks, ChunkSize](FHttpRequestPtr HttpRequest, FHttpResponsePtr HttpResponse, bool bSucceeded) {
+			CHECK(HttpResponse->GetContentLength() == Chunks * ChunkSize);
+			CHECK(bSucceeded);
+			CHECK(HttpResponse->GetResponseCode() == 200);
+		});
+		HttpRequest->ProcessRequest();
+	}
+}
+
 TEST_CASE_METHOD(FWaitUntilCompleteHttpFixture, "Streaming http upload - gold path.", HTTP_TAG)
 {
 	TSharedRef<IHttpRequest, ESPMode::ThreadSafe> HttpRequest = HttpModule->CreateRequest();
diff --git a/Engine/Source/Runtime/Online/HTTP/Private/IHttpThreadedRequest.h b/Engine/Source/Runtime/Online/HTTP/Private/IHttpThreadedRequest.h
--- a/Engine/Source/Runtime/Online/HTTP/Private/IHttpThreadedRequest.h
+++ b/Engine/Source/Runtime/Online/HTTP/Private/IHttpThreadedRequest.h
@@ -13,7 +13,7 @@ public:
 	virtual bool IsThreadedRequestComplete() = 0;
 	virtual void TickThreadedRequest(float DeltaSeconds) = 0;
 
-	// Called on game thread
+	// Can be called on game thread or http thread depend on the delegate thread policy
 	virtual void FinishRequest() = 0;
 
 protected:
