Add test case to test file upload through file name;
Add test case to test redirect;
Minor refactor for test code: Extract UrlStreamDownload function; Remove ESPMode::ThreadSafe when define TSharedRef.

diff --git a/Engine/Source/Programs/WebTests/Private/TestHttp.cpp b/Engine/Source/Programs/WebTests/Private/TestHttp.cpp
--- a/Engine/Source/Programs/WebTests/Private/TestHttp.cpp
+++ b/Engine/Source/Programs/WebTests/Private/TestHttp.cpp
@@ -56,6 +56,7 @@ public:
 	const FString UrlBase() const { return FString::Format(TEXT("http://{0}:{1}"), { *WebServerIp, 8000 }); }
 	const FString UrlHttpTests() const { return FString::Format(TEXT("{0}/webtests/httptests"), { *UrlBase() }); }
 	const FString UrlToTestMethods() const { return FString::Format(TEXT("{0}/methods"), { *UrlHttpTests() }); }
+	const FString UrlStreamDownload(uint32 Chunks, uint32 ChunkSize) { return FString::Format(TEXT("{0}/streaming_download/{1}/{2}/"), { *UrlHttpTests(), Chunks, ChunkSize }); }
 
 	FString WebServerIp;
 	FHttpModule* HttpModule;
@@ -72,7 +73,7 @@ TEST_CASE_METHOD(FHttpModuleTestFixture, "Shutdown http module without issue whe
 	{
 		IHttpRequest* LeakingHttpRequest = FPlatformHttp::ConstructRequest(); // Leaking in purpose to make sure it's ok
 
-		TSharedRef<IHttpRequest, ESPMode::ThreadSafe> HttpRequest = HttpModule->CreateRequest();
+		TSharedRef<IHttpRequest> HttpRequest = HttpModule->CreateRequest();
 		HttpRequest->SetURL(UrlToTestMethods());
 		HttpRequest->SetVerb(TEXT("PUT"));
 		// TODO: Use some shared data, like cookie, openssl session etc.
@@ -123,7 +124,7 @@ TEST_CASE_METHOD(FWaitUntilQuitFromTestFixture, "Http request can be reused", HT
 
 		uint32 Chunks = 3;
 		uint32 ChunkSize = 1024;
-		HttpRequest->SetURL(FString::Format(TEXT("{0}/streaming_download/{1}/{2}/"), { *UrlHttpTests(), Chunks, ChunkSize }));
+		HttpRequest->SetURL(UrlStreamDownload(Chunks, ChunkSize));
 		HttpRequest->SetVerb(TEXT("GET"));
 		HttpRequest->OnProcessRequestComplete().BindLambda([this, Chunks, ChunkSize](FHttpRequestPtr HttpRequest, FHttpResponsePtr HttpResponse, bool bSucceeded) {
 			CHECK(bSucceeded);
@@ -183,7 +184,7 @@ public:
 
 TEST_CASE_METHOD(FWaitUntilCompleteHttpFixture, "Http Methods", HTTP_TAG)
 {
-	TSharedRef<IHttpRequest, ESPMode::ThreadSafe> HttpRequest = HttpModule->CreateRequest();
+	TSharedRef<IHttpRequest> HttpRequest = HttpModule->CreateRequest();
 	HttpRequest->SetURL(UrlToTestMethods());
 
 	SECTION("GET")
@@ -213,7 +214,7 @@ TEST_CASE_METHOD(FWaitUntilCompleteHttpFixture, "Http Methods", HTTP_TAG)
 
 TEST_CASE_METHOD(FWaitUntilCompleteHttpFixture, "Get large response content without chunks", HTTP_TAG)
 {
-	TSharedRef<IHttpRequest, ESPMode::ThreadSafe> HttpRequest = HttpModule->CreateRequest();
+	TSharedRef<IHttpRequest> HttpRequest = HttpModule->CreateRequest();
 	HttpRequest->SetURL(FString::Format(TEXT("{0}/get_large_response_without_chunks/{1}/"), { *UrlHttpTests(), 1024 * 1024/*bytes_number*/}));
 	HttpRequest->SetVerb(TEXT("GET"));
 	HttpRequest->OnProcessRequestComplete().BindLambda([](FHttpRequestPtr HttpRequest, FHttpResponsePtr HttpResponse, bool bSucceeded) {
@@ -226,7 +227,7 @@ TEST_CASE_METHOD(FWaitUntilCompleteHttpFixture, "Get large response content with
 
 TEST_CASE_METHOD(FWaitUntilCompleteHttpFixture, "Http request connect timeout", HTTP_TAG)
 {
-	TSharedRef<IHttpRequest, ESPMode::ThreadSafe> HttpRequest = HttpModule->CreateRequest();
+	TSharedRef<IHttpRequest> HttpRequest = HttpModule->CreateRequest();
 	HttpRequest->SetURL(UrlWithInvalidPortToTestConnectTimeout());
 	HttpRequest->SetVerb(TEXT("GET"));
 	HttpRequest->SetTimeout(7);
@@ -249,8 +250,8 @@ TEST_CASE_METHOD(FWaitUntilCompleteHttpFixture, "Streaming http download", HTTP_
 	uint32 Chunks = 3;
 	uint32 ChunkSize = 1024*1024;
 
-	TSharedRef<IHttpRequest, ESPMode::ThreadSafe> HttpRequest = HttpModule->CreateRequest();
-	HttpRequest->SetURL(FString::Format(TEXT("{0}/streaming_download/{1}/{2}/"), { *UrlHttpTests(), Chunks, ChunkSize }));
+	TSharedRef<IHttpRequest> HttpRequest = HttpModule->CreateRequest();
+	HttpRequest->SetURL(UrlStreamDownload(Chunks, ChunkSize));
 	HttpRequest->SetVerb(TEXT("GET"));
 
 	TSharedRef<int64> TotalBytesReceived = MakeShared<int64>(0);
@@ -373,7 +374,7 @@ TEST_CASE_METHOD(FWaitUntilCompleteHttpFixture, "Streaming http download", HTTP_
 			FileToWrite->FlushCache();
 			FileToWrite->Close();
 
-			TSharedRef<FArchive, ESPMode::ThreadSafe> FileToRead = MakeShareable(IFileManager::Get().CreateFileReader(*Filename));
+			TSharedRef<FArchive> FileToRead = MakeShareable(IFileManager::Get().CreateFileReader(*Filename));
 			CHECK(FileToRead->TotalSize() == Chunks * ChunkSize);
 
 			IFileManager::Get().Delete(*Filename);
@@ -390,8 +391,8 @@ TEST_CASE_METHOD(FWaitUntilCompleteHttpFixture, "Can run parallel stream downloa
 
 	for (int i = 0; i < 3; ++i)
 	{
-		TSharedRef<IHttpRequest, ESPMode::ThreadSafe> HttpRequest = HttpModule->CreateRequest();
-		HttpRequest->SetURL(FString::Format(TEXT("{0}/streaming_download/{1}/{2}/"), { *UrlHttpTests(), Chunks, ChunkSize }));
+		TSharedRef<IHttpRequest> HttpRequest = HttpModule->CreateRequest();
+		HttpRequest->SetURL(UrlStreamDownload(Chunks, ChunkSize));
 		HttpRequest->SetVerb(TEXT("GET"));
 		HttpRequest->OnProcessRequestComplete().BindLambda([Chunks, ChunkSize](FHttpRequestPtr HttpRequest, FHttpResponsePtr HttpResponse, bool bSucceeded) {
 			CHECK(HttpResponse->GetContentLength() == Chunks * ChunkSize);
@@ -402,16 +403,16 @@ TEST_CASE_METHOD(FWaitUntilCompleteHttpFixture, "Can run parallel stream downloa
 	}
 }
 
-TEST_CASE_METHOD(FWaitUntilCompleteHttpFixture, "Streaming http upload - gold path.", HTTP_TAG)
+TEST_CASE_METHOD(FWaitUntilCompleteHttpFixture, "Streaming http upload from memory", HTTP_TAG)
 {
-	TSharedRef<IHttpRequest, ESPMode::ThreadSafe> HttpRequest = HttpModule->CreateRequest();
-	HttpRequest->SetURL(FString::Format(TEXT("{0}/streaming_upload"), { *UrlHttpTests() }));
+	TSharedRef<IHttpRequest> HttpRequest = HttpModule->CreateRequest();
+	HttpRequest->SetURL(FString::Format(TEXT("{0}/streaming_upload_post"), { *UrlHttpTests() }));
 	HttpRequest->SetVerb(TEXT("POST"));
 
 	const char* BoundaryLabel = "test_http_boundary";
 	HttpRequest->SetHeader(TEXT("Content-Type"), FString::Format(TEXT("multipart/form-data; boundary={0}"), { BoundaryLabel }));
 
-	// Not really reading file here in order to simplify the test flow. It will be sent by chunks in http request
+	// Data will be sent by chunks in http request
 	const uint32 FileSize = 10*1024*1024;
 	char* FileData = (char*)FMemory::Malloc(FileSize + 1);
 	FMemory::Memset(FileData, 'd', FileSize);
@@ -445,6 +446,46 @@ TEST_CASE_METHOD(FWaitUntilCompleteHttpFixture, "Streaming http upload - gold pa
 	HttpRequest->ProcessRequest();
 }
 
+TEST_CASE_METHOD(FWaitUntilCompleteHttpFixture, "Streaming http upload from file by PUT can work well", HTTP_TAG)
+{
+	FString Filename = FString(FPlatformProcess::UserSettingsDir()) / TEXT("TestStreamUpload.dat");
+
+	FArchive* RawFile = IFileManager::Get().CreateFileWriter(*Filename);
+	CHECK(RawFile != nullptr);
+	TSharedRef<FArchive> FileToWrite = MakeShareable(RawFile);
+	const uint64 FileSize = 5*1024*1024; // 5MB
+	char* FileData = (char*)FMemory::Malloc(FileSize);
+	FMemory::Memset(FileData, 'd', FileSize);
+	FileToWrite->Serialize(FileData, FileSize);
+	FileToWrite->FlushCache();
+	FileToWrite->Close();
+	FMemory::Free(FileData);
+
+	TSharedRef<IHttpRequest> HttpRequest = HttpModule->CreateRequest();
+	HttpRequest->SetURL(FString::Format(TEXT("{0}/streaming_upload_put"), { *UrlHttpTests() }));
+	HttpRequest->SetVerb(TEXT("PUT"));
+	HttpRequest->SetHeader(TEXT("Content-Disposition"), TEXT("attachment;filename=TestStreamUpload.dat"));
+	HttpRequest->SetContentAsStreamedFile(Filename);
+	HttpRequest->OnProcessRequestComplete().BindLambda([Filename](FHttpRequestPtr HttpRequest, FHttpResponsePtr HttpResponse, bool bSucceeded) {
+		CHECK(bSucceeded);
+		CHECK(HttpResponse->GetResponseCode() == 200);
+		IFileManager::Get().Delete(*Filename);
+	});
+	HttpRequest->ProcessRequest();
+}
+
+TEST_CASE_METHOD(FWaitUntilCompleteHttpFixture, "Redirect enabled by default and can work well", HTTP_TAG)
+{
+	TSharedRef<IHttpRequest> HttpRequest = HttpModule->CreateRequest();
+	HttpRequest->SetURL(FString::Format(TEXT("{0}/redirect_from"), { *UrlHttpTests() }));
+	HttpRequest->SetVerb(TEXT("GET"));
+	HttpRequest->OnProcessRequestComplete().BindLambda([](FHttpRequestPtr HttpRequest, FHttpResponsePtr HttpResponse, bool bSucceeded) {
+		CHECK(bSucceeded);
+		CHECK(HttpResponse->GetResponseCode() == 200);
+	});
+	HttpRequest->ProcessRequest();
+}
+
 class FWaitThreadedHttpFixture : public FWaitUntilCompleteHttpFixture, public FRunnable
 {
 public:
@@ -476,7 +517,7 @@ TEST_CASE_METHOD(FWaitThreadedHttpFixture, "Http streaming download request can
 {
 	ThreadCallback.BindLambda([this]() {
 		TSharedRef<IHttpRequest> HttpRequest = HttpModule->CreateRequest();
-		HttpRequest->SetURL(FString::Format(TEXT("{0}/streaming_download/{1}/{2}/"), { *UrlHttpTests(), 3/*chunks*/, 1024/*chunk_size*/}));
+		HttpRequest->SetURL(UrlStreamDownload(3/*Chunks*/, 1024/*ChunkSize*/));
 		HttpRequest->SetVerb(TEXT("GET"));
 		HttpRequest->SetDelegateThreadPolicy(EHttpRequestDelegateThreadPolicy::CompleteOnHttpThread);
 
@@ -526,7 +567,7 @@ TEST_CASE_METHOD(FWaitUntilCompleteHttpFixture, "Http request pre check will fai
 	// Pre check will fail when domain is not allowed
 	UE::TestHttp::SetupURLRequestFilter(HttpModule);
 
-	TSharedRef<IHttpRequest, ESPMode::ThreadSafe> HttpRequest = HttpModule->CreateRequest();
+	TSharedRef<IHttpRequest> HttpRequest = HttpModule->CreateRequest();
 	HttpRequest->SetVerb(TEXT("GET"));
 	HttpRequest->SetURL(UrlToTestMethods());
 
@@ -555,7 +596,7 @@ TEST_CASE_METHOD(FWaitThreadedHttpFixture, "Threaded http request pre check will
 		// Pre check will fail when domain is not allowed
 		UE::TestHttp::SetupURLRequestFilter(HttpModule);
 
-		TSharedRef<IHttpRequest, ESPMode::ThreadSafe> HttpRequest = HttpModule->CreateRequest();
+		TSharedRef<IHttpRequest> HttpRequest = HttpModule->CreateRequest();
 		HttpRequest->SetVerb(TEXT("GET"));
 		HttpRequest->SetURL(UrlToTestMethods());
 
diff --git a/Engine/Source/Programs/WebTests/WebServer/httptests/urls.py b/Engine/Source/Programs/WebTests/WebServer/httptests/urls.py
--- a/Engine/Source/Programs/WebTests/WebServer/httptests/urls.py
+++ b/Engine/Source/Programs/WebTests/WebServer/httptests/urls.py
@@ -7,5 +7,8 @@ urlpatterns = [
     path("httptests/get_large_response_without_chunks/<int:bytes_number>/", views.get_large_response_without_chunks, name="get_large_response_without_chunks"),
     path("httptests/nonstreaming_receivetimeout/<int:wait_time>/", views.nonstreaming_receivetimeout, name="nonstreaming_receivetimeout"),
     path("httptests/streaming_download/<int:chunks>/<int:chunk_size>/", views.streaming_download, name="streaming_download"),
-    path("httptests/streaming_upload", views.streaming_upload, name="streaming_upload"),
+    path("httptests/streaming_upload_post", views.streaming_upload_post, name="streaming_upload_post"),
+    path("httptests/streaming_upload_put", views.streaming_upload_put, name="streaming_upload_put"),
+    path("httptests/redirect_from", views.redirect_from, name="redirect_from"),
+    path("httptests/redirect_to", views.redirect_to, name="redirect_to"),
 ]
diff --git a/Engine/Source/Programs/WebTests/WebServer/httptests/views.py b/Engine/Source/Programs/WebTests/WebServer/httptests/views.py
--- a/Engine/Source/Programs/WebTests/WebServer/httptests/views.py
+++ b/Engine/Source/Programs/WebTests/WebServer/httptests/views.py
@@ -1,8 +1,9 @@
-from django.shortcuts import render
+from django.shortcuts import render, redirect
 from django.http.response import JsonResponse
 from django.http import JsonResponse
 from django.http import StreamingHttpResponse
-from rest_framework.decorators import api_view
+from rest_framework.decorators import api_view, parser_classes
+from rest_framework.parsers import FileUploadParser
 import time
 import logging
 
@@ -32,9 +33,7 @@ def streaming_download(request, chunks, chunk_size):
     response['Content-Length'] = chunks * chunk_size
     return response
 
-@api_view(['POST'])
-def streaming_upload(request):
-    f = request.FILES["file"]
+def output_file_info(f):
     if f.multiple_chunks():
         # big file
         total_chunks = 0
@@ -46,4 +45,24 @@ def streaming_upload(request):
     else:
         # small file
         logger.info("Received the whole file: %d bytes", f.size)
+
+@api_view(['POST'])
+def streaming_upload_post(request):
+    f = request.FILES['file']
+    output_file_info(f)
+    return JsonResponse({})
+
+@api_view(['PUT'])
+@parser_classes((FileUploadParser,))
+def streaming_upload_put(request):
+    f = request.data['file']
+    output_file_info(f)
+    return JsonResponse({})
+
+@api_view(['GET'])
+def redirect_from(request):
+    return redirect('redirect_to')
+
+@api_view(['GET'])
+def redirect_to(request):
     return JsonResponse({})
