Block main thread in the test when verifying http non-game thread support.

diff --git a/Engine/Source/Programs/WebTests/Private/TestHttp.cpp b/Engine/Source/Programs/WebTests/Private/TestHttp.cpp
--- a/Engine/Source/Programs/WebTests/Private/TestHttp.cpp
+++ b/Engine/Source/Programs/WebTests/Private/TestHttp.cpp
@@ -606,9 +606,9 @@ public:
 
 TEST_CASE_METHOD(FWaitThreadedHttpFixture, "Http streaming download request can work in non game thread", HTTP_TAG)
 {
-	// TODO: Block main thread here to verify
+	std::atomic<bool> bBlockingGameThreadTick = true;
 
-	ThreadCallback.BindLambda([this]() {
+	ThreadCallback.BindLambda([this, &bBlockingGameThreadTick]() {
 		TSharedRef<IHttpRequest> HttpRequest = HttpModule->CreateRequest();
 		HttpRequest->SetURL(UrlStreamDownload(3/*Chunks*/, 1024/*ChunkSize*/));
 		HttpRequest->SetVerb(TEXT("GET"));
@@ -625,17 +625,23 @@ TEST_CASE_METHOD(FWaitThreadedHttpFixture, "Http streaming download request can
 		};
 		CHECK(HttpRequest->SetResponseBodyReceiveStream(MakeShared<FTestHttpReceiveStream>()));
 
-		HttpRequest->OnProcessRequestComplete().BindLambda([](FHttpRequestPtr HttpRequest, FHttpResponsePtr HttpResponse, bool bSucceeded) {
+		HttpRequest->OnProcessRequestComplete().BindLambda([&bBlockingGameThreadTick](FHttpRequestPtr HttpRequest, FHttpResponsePtr HttpResponse, bool bSucceeded) {
 			// EHttpRequestDelegateThreadPolicy::CompleteOnHttpThread was used, so not in game thread here
 			CHECK(!IsInGameThread());
 			CHECK(bSucceeded);
 			CHECK(HttpResponse->GetResponseCode() == 200);
+			bBlockingGameThreadTick = false;
 		});
 
 		HttpRequest->ProcessRequest();
 	});
 
 	StartTestHttpThread();
+
+	while (bBlockingGameThreadTick)
+	{
+		FPlatformProcess::Sleep(TickFrequency);
+	}
 }
 
 namespace UE
