Adding https tests in webtests

diff --git a/Engine/Source/Programs/WebTests/Private/TestHttp.cpp b/Engine/Source/Programs/WebTests/Private/TestHttp.cpp
--- a/Engine/Source/Programs/WebTests/Private/TestHttp.cpp
+++ b/Engine/Source/Programs/WebTests/Private/TestHttp.cpp
@@ -32,8 +32,8 @@ class FHttpModuleTestFixture
 public:
 	FHttpModuleTestFixture()
 		: WebServerIp(TEXT("127.0.0.1"))
-		, WebServerPort(8000)
 		, bRunHeavyTests(false)
+		, bTestHttps(false)
 	{
 		ParseSettingsFromCommandLine();
 
@@ -55,17 +55,16 @@ public:
 		FParse::Bool(FCommandLine::Get(), TEXT("run_heavy_tests"), bRunHeavyTests);
 	}
 
-	const FString UrlWithInvalidPortToTestConnectTimeout() const { return FString::Format(TEXT("http://{0}:{1}"), { *WebServerIp, 8765 }); }
-	const FString UrlBase() const { return FString::Format(TEXT("http://{0}:{1}"), { *WebServerIp, WebServerPort }); }
+	const FString UrlWithInvalidPortToTestConnectTimeout() const { return FString::Format(TEXT("{0}://{1}:{2}"), { this->bTestHttps ? TEXT("https") : TEXT("http"), *WebServerIp, 8765 }); }
+	const FString UrlBase() const { return FString::Format(TEXT("{0}://{1}:{2}"), { this->bTestHttps ? TEXT("https") : TEXT("http"), *WebServerIp, this->bTestHttps ? TEXT("8001") : TEXT("8000") }); }
 	const FString UrlHttpTests() const { return FString::Format(TEXT("{0}/webtests/httptests"), { *UrlBase() }); }
 	const FString UrlToTestMethods() const { return FString::Format(TEXT("{0}/methods"), { *UrlHttpTests() }); }
 	const FString UrlStreamDownload(uint32 Chunks, uint32 ChunkSize) { return FString::Format(TEXT("{0}/streaming_download/{1}/{2}/"), { *UrlHttpTests(), Chunks, ChunkSize }); }
 
 	FString WebServerIp;
-	uint32 WebServerPort;
 	FHttpModule* HttpModule;
-
 	bool bRunHeavyTests;
+	bool bTestHttps;
 };
 
 TEST_CASE_METHOD(FHttpModuleTestFixture, "Shutdown http module without issue when there are ongoing http requests.", HTTP_TAG)
diff --git a/Engine/Source/Programs/WebTests/WebServer/requirements.txt b/Engine/Source/Programs/WebTests/WebServer/requirements.txt
--- a/Engine/Source/Programs/WebTests/WebServer/requirements.txt
+++ b/Engine/Source/Programs/WebTests/WebServer/requirements.txt
@@ -1,2 +1,3 @@
 django==4.1.7
+django-sslserver==0.22
 djangorestframework==3.14.0
diff --git a/Engine/Source/Programs/WebTests/WebServer/runserver.bat b/Engine/Source/Programs/WebTests/WebServer/runserver.bat
--- a/Engine/Source/Programs/WebTests/WebServer/runserver.bat
+++ b/Engine/Source/Programs/WebTests/WebServer/runserver.bat
@@ -12,4 +12,5 @@ if exist env (
 	)
 )
 
-python manage.py runserver 0.0.0.0:8000
+start "HTTP Server" python manage.py runserver 0.0.0.0:8000
+@REM start "HTTPS Server" python manage.py runsslserver 0.0.0.0:8001 --certificate ssl/https_server_certificate.crt --key ssl/https_server.key
diff --git a/Engine/Source/Programs/WebTests/WebServer/runserver.sh b/Engine/Source/Programs/WebTests/WebServer/runserver.sh
--- a/Engine/Source/Programs/WebTests/WebServer/runserver.sh
+++ b/Engine/Source/Programs/WebTests/WebServer/runserver.sh
@@ -1,13 +1,14 @@
 if [ -d "env" ]; then
-	source ./env/Scripts/activate.sh
+    source ./env/Scripts/activate.sh
 else
-	run ../../../../Binaries/ThirdParty/Python3/Mac/python3 -m venv env
-	source ./env/Scripts/activate.sh
-	run pip install -r requirements.txt
-	if [ $? -ne 0 ]; then
-		rm -rf env
-		exit $?
-	fi
+    run ../../../../Binaries/ThirdParty/Python3/Mac/python3 -m venv env
+    source ./env/Scripts/activate.sh
+    run pip install -r requirements.txt
+    if [ $? -ne 0 ]; then
+        rm -rf env
+        exit $?
+    fi
 fi
 
-cmd python manage.py runserver 0.0.0.0:8000
+python manage.py runserver 0.0.0.0:8000 &
+#python manage.py runsslserver 0.0.0.0:8001 --certificate ssl/https_server_certificate.crt --key ssl/https_server.key &
diff --git a/Engine/Source/Programs/WebTests/WebServer/webtests/settings.py b/Engine/Source/Programs/WebTests/WebServer/webtests/settings.py
--- a/Engine/Source/Programs/WebTests/WebServer/webtests/settings.py
+++ b/Engine/Source/Programs/WebTests/WebServer/webtests/settings.py
@@ -38,6 +38,7 @@ INSTALLED_APPS = [
     'django.contrib.sessions',
     'django.contrib.messages',
     'django.contrib.staticfiles',
+    'sslserver',
 ]
 
 MIDDLEWARE = [
@@ -116,7 +117,7 @@ USE_TZ = True
 # Static files (CSS, JavaScript, Images)
 # https://docs.djangoproject.com/en/4.1/howto/static-files/
 
-STATIC_URL = 'static/'
+STATIC_URL = '/static/'
 
 # Default primary key field type
 # https://docs.djangoproject.com/en/4.1/ref/settings/#default-auto-field
@@ -128,3 +129,14 @@ LOGGING = {
     "disable_existing_loggers": False,  # retain the default loggers
     "level": "DEBUG",
 }
+
+# HTTPS settings
+SECURE_SSL_REDIRECT = False
+SECURE_HSTS_SECONDS = 31536000  # a year
+SECURE_HSTS_INCLUDE_SUBDOMAINS = True
+SECURE_HSTS_PRELOAD = True
+SESSION_COOKIE_SECURE = True
+CSRF_COOKIE_SECURE = True
+SECURE_PROXY_SSL_HEADER = ('HTTP_X_FORWARDED_PROTO', 'https')
+SECURE_SSL_CERTIFICATE = 'ssl/https_server_certificate.crt'
+SECURE_SSL_KEY = 'ssl/https_server.key'
