Add socket dependency to webtests project.

diff --git a/Engine/Source/Programs/WebTests/Private/WebLLTModule.cpp b/Engine/Source/Programs/WebTests/Private/WebLLTModule.cpp
new file mode 100644
--- /dev/null
+++ b/Engine/Source/Programs/WebTests/Private/WebLLTModule.cpp
@@ -0,0 +1,51 @@
+// Copyright Epic Games, Inc. All Rights Reserved.
+
+#include "CoreTypes.h"
+
+#include "Containers/Array.h"
+#include "Containers/UnrealString.h"
+#include "Modules/ModuleManager.h"
+#include "TestRunner.h"
+
+namespace UE::WebTests
+{
+
+TArray<FString> GetRequiredModules()
+{
+	static TArray<FString> Modules({
+		FString(TEXT("Sockets"))
+	});
+
+	return Modules;
+}
+
+void InitializeWebTests()
+{
+	for (const FString& ModuleName : GetRequiredModules())
+	{
+		FModuleManager::LoadModulePtr<IModuleInterface>(*ModuleName);
+	}
+}
+
+void CleanupWebTests()
+{
+	for (const FString& ModuleName : GetRequiredModules())
+	{
+		if (IModuleInterface* Module = FModuleManager::Get().GetModule(*ModuleName))
+		{
+			Module->ShutdownModule();
+		}
+	}
+}
+
+class WebTestsGlobalSetup
+{
+public:
+	WebTestsGlobalSetup()
+	{
+		FTestDelegates::GetGlobalSetup().BindStatic(&InitializeWebTests);
+		FTestDelegates::GetGlobalTeardown().BindStatic(&CleanupWebTests);
+	}
+} GWebTestsGlobalSetup;
+
+} // UE::WebTests
