Make sure progress callback can be called on http thread;
Make sure HttpManagerRequestAdded and HttpManagerRequestCompleted callback can be called on http thread.

From now on the delegates bind to http manager through SetRequestAddedDelegate and SetRequestCompletedDelegate have to be changed to be thread-safe.

diff --git a/Engine/Source/Programs/WebTests/Private/TestHttp.cpp b/Engine/Source/Programs/WebTests/Private/TestHttp.cpp
--- a/Engine/Source/Programs/WebTests/Private/TestHttp.cpp
+++ b/Engine/Source/Programs/WebTests/Private/TestHttp.cpp
@@ -103,10 +103,10 @@ public:
 
 	~FWaitUntilQuitFromTestFixture()
 	{
-		WaitUntilAllHttpRequestsComplete();
+		WaitUntilQuitFromTest();
 	}
 
-	void WaitUntilAllHttpRequestsComplete()
+	void WaitUntilQuitFromTest()
 	{
 		while (!bQuitRequested)
 		{
@@ -169,10 +169,7 @@ public:
 
 	void OnRequestCompleted(const FHttpRequestRef& Request)
 	{
-		if (ensure(OngoingRequests > 0))
-		{
-			--OngoingRequests;
-		}
+		ensure(--OngoingRequests >= 0);
 	}
 
 	void WaitUntilAllHttpRequestsComplete()
@@ -184,7 +181,7 @@ public:
 		}
 	}
 
-	uint32 OngoingRequests = 0;
+	std::atomic<int32> OngoingRequests = 0;
 	float TickFrequency = 1.0f / 60; /*60 FPS*/;
 };
 
@@ -685,6 +682,39 @@ TEST_CASE_METHOD(FWaitThreadedHttpFixture, "Http streaming download request can
 	ThreadedHttpRunnable.StartTestHttpThread(true/*bBlockGameThread*/);
 }
 
+TEST_CASE_METHOD(FWaitThreadedHttpFixture, "Http download request progress callback can be received in http thread", HTTP_TAG)
+{
+	std::atomic<bool> bRequestProgressTriggered = false;
+	ThreadedHttpRunnable.OnRunFromThread().BindLambda([this, &bRequestProgressTriggered]() {
+		TSharedRef<IHttpRequest> HttpRequest = HttpModule->CreateRequest();
+		HttpRequest->SetURL(UrlStreamDownload(10/*Chunks*/, 1024*1024/*ChunkSize*/));
+		HttpRequest->SetVerb(TEXT("GET"));
+
+		HttpRequest->SetDelegateThreadPolicy(EHttpRequestDelegateThreadPolicy::CompleteOnHttpThread);
+		HttpRequest->OnRequestProgress64().BindLambda([this, &bRequestProgressTriggered](FHttpRequestPtr Request, uint64 /*BytesSent*/, uint64 BytesReceived) {
+			if (!bRequestProgressTriggered)
+			{
+				// Only do these checks once, because when http request complete, this callback also get triggered
+				CHECK(BytesReceived > 0);
+				CHECK(BytesReceived < 10/*Chunks*/ * 1024*1024/*ChunkSize*/);
+				CHECK(!IsInGameThread());
+				CHECK(Request->GetStatus() == EHttpRequestStatus::Processing);
+				bRequestProgressTriggered = true;
+			}
+		});
+		HttpRequest->OnProcessRequestComplete().BindLambda([this](FHttpRequestPtr /*HttpRequest*/, FHttpResponsePtr /*HttpResponse */, bool bSucceeded) {
+			CHECK(bSucceeded);
+			ThreadedHttpRunnable.UnblockGameThread();
+		});
+
+		HttpRequest->ProcessRequest();
+	});
+
+	ThreadedHttpRunnable.StartTestHttpThread(true/*bBlockGameThread*/);
+
+	CHECK(bRequestProgressTriggered);
+}
+
 namespace UE
 {
 namespace TestHttp
@@ -879,3 +909,7 @@ TEST_CASE_METHOD(FRetryHttpManagerThreadedRequestsFixture, "Retry manager is thr
 
 
 // TODO: Add cancel test, with multiple cancel calls
+
+// TODO: Add a test case to make sure once Request completed and destroyed, if the user code still keeps the shared ptr of Response, any call to Response shouldn't have dependency to Request
+
+// TODO: Add test case to validate header received callback can be received in http/game thread, and can be received before the request complete
diff --git a/Engine/Source/Programs/WebTests/Private/TestWebSockets.cpp b/Engine/Source/Programs/WebTests/Private/TestWebSockets.cpp
--- a/Engine/Source/Programs/WebTests/Private/TestWebSockets.cpp
+++ b/Engine/Source/Programs/WebTests/Private/TestWebSockets.cpp
@@ -2,6 +2,7 @@
 #include "Containers/BackgroundableTicker.h"
 #include "CoreMinimal.h"
 #include "HAL/PlatformProcess.h"
+#include "Http.h"
 #include "Misc/CommandLine.h"
 #include "WebSocketsModule.h"
 #include "IWebSocket.h"
@@ -30,6 +31,11 @@ public:
 	{
 		ParseSettingsFromCommandLine();
 
+		// Init HTTP module because websockets module has dependency on it when get proxy
+		HttpModule = new FHttpModule();
+		IModuleInterface* HttpModuleInterface = HttpModule;
+		HttpModuleInterface->StartupModule();
+
 		WebSocketsModule = new FWebSocketsModule();
 		IModuleInterface* Module = WebSocketsModule;
 		Module->StartupModule();
@@ -37,9 +43,13 @@ public:
 
 	virtual ~FWebSocketsModuleTestFixture()
 	{
-		IModuleInterface* Module = WebSocketsModule;
-		Module->ShutdownModule();
-		delete Module;
+		IModuleInterface* WebSocketsModuleInterface = WebSocketsModule;
+		WebSocketsModuleInterface->ShutdownModule();
+		delete WebSocketsModuleInterface;
+
+		IModuleInterface* HttpModuleInterface = HttpModule;
+		HttpModuleInterface->ShutdownModule();
+		delete HttpModuleInterface;
 	}
 
 	void ParseSettingsFromCommandLine()
@@ -54,6 +64,7 @@ public:
 	FString WebServerIp;
 	uint32 WebServerWebSocketsPort;
 	FWebSocketsModule* WebSocketsModule;
+	FHttpModule* HttpModule;
 };
 
 class FRunUntilQuitRequestedFixture : public FWebSocketsModuleTestFixture
diff --git a/Engine/Source/Runtime/Online/HTTP/Private/Apple/AppleHTTPNSUrlSession.cpp b/Engine/Source/Runtime/Online/HTTP/Private/Apple/AppleHTTPNSUrlSession.cpp
--- a/Engine/Source/Runtime/Online/HTTP/Private/Apple/AppleHTTPNSUrlSession.cpp
+++ b/Engine/Source/Runtime/Online/HTTP/Private/Apple/AppleHTTPNSUrlSession.cpp
@@ -811,6 +811,14 @@ const FHttpResponsePtr FAppleHttpNSUrlSessionRequest::GetResponse() const
 }
 
 void FAppleHttpNSUrlSessionRequest::Tick(float DeltaSeconds)
+{
+	if (DelegateThreadPolicy == EHttpRequestDelegateThreadPolicy::CompleteOnGameThread)
+	{
+		CheckProgressDelegate();
+	}
+}
+
+void FAppleHttpNSUrlSessionRequest::CheckProgressDelegate()
 {
 	if (Response.IsValid() && (CompletionStatus == EHttpRequestStatus::Processing || Response->HadError()))
 	{
@@ -847,6 +855,11 @@ bool FAppleHttpNSUrlSessionRequest::IsThreadedRequestComplete()
 void FAppleHttpNSUrlSessionRequest::TickThreadedRequest(float DeltaSeconds)
 {
 	ElapsedTime += DeltaSeconds;
+
+	if (DelegateThreadPolicy == EHttpRequestDelegateThreadPolicy::CompleteOnHttpThread)
+	{
+		CheckProgressDelegate();
+	}
 }
 
 /****************************************************************************
diff --git a/Engine/Source/Runtime/Online/HTTP/Private/Apple/AppleHTTPNSUrlSession.h b/Engine/Source/Runtime/Online/HTTP/Private/Apple/AppleHTTPNSUrlSession.h
--- a/Engine/Source/Runtime/Online/HTTP/Private/Apple/AppleHTTPNSUrlSession.h
+++ b/Engine/Source/Runtime/Online/HTTP/Private/Apple/AppleHTTPNSUrlSession.h
@@ -77,6 +77,10 @@ public:
 
 
 private:
+	/**
+	 * Trigger the request progress delegate if progress has changed
+	 */
+	void CheckProgressDelegate();
 
 	/**
 	 * Create the session connection and initiate the web request
diff --git a/Engine/Source/Runtime/Online/HTTP/Private/Curl/CurlHttp.cpp b/Engine/Source/Runtime/Online/HTTP/Private/Curl/CurlHttp.cpp
--- a/Engine/Source/Runtime/Online/HTTP/Private/Curl/CurlHttp.cpp
+++ b/Engine/Source/Runtime/Online/HTTP/Private/Curl/CurlHttp.cpp
@@ -471,7 +471,15 @@ size_t FCurlHttpRequest::StaticReceiveResponseBodyCallback(void* Ptr, size_t Siz
 
 	// dispatch
 	FCurlHttpRequest* Request = reinterpret_cast<FCurlHttpRequest*>(UserData);
-	return Request->ReceiveResponseBodyCallback(Ptr, SizeInBlocks, BlockSizeInBytes);	
+
+	size_t Result = Request->ReceiveResponseBodyCallback(Ptr, SizeInBlocks, BlockSizeInBytes);	
+
+	if (Request->DelegateThreadPolicy == EHttpRequestDelegateThreadPolicy::CompleteOnHttpThread)
+	{
+		Request->CheckProgressDelegate();
+	}
+
+	return Result;
 }
 
 size_t FCurlHttpRequest::StaticDebugCallback(CURL * Handle, curl_infotype DebugInfoType, char * DebugInfo, size_t DebugInfoSize, void* UserData)
@@ -1106,10 +1114,9 @@ const FHttpResponsePtr FCurlHttpRequest::GetResponse() const
 
 void FCurlHttpRequest::Tick(float DeltaSeconds)
 {
-	CheckProgressDelegate();
-
 	if (DelegateThreadPolicy == EHttpRequestDelegateThreadPolicy::CompleteOnGameThread)
 	{
+		CheckProgressDelegate();
 		BroadcastNewlyReceivedHeaders();
 	}
 }
@@ -1154,7 +1161,6 @@ void FCurlHttpRequest::BroadcastNewlyReceivedHeaders()
 
 void FCurlHttpRequest::BroadcastNewlyReceivedHeader(const FString& HeaderKey, const FString& HeaderValue)
 {
-	check(IsInGameThread() || DelegateThreadPolicy == EHttpRequestDelegateThreadPolicy::CompleteOnHttpThread);
 	check(Response);
 
 	const constexpr FStringView Seperator(TEXTVIEW(", "));
diff --git a/Engine/Source/Runtime/Online/HTTP/Private/HttpThread.cpp b/Engine/Source/Runtime/Online/HTTP/Private/HttpThread.cpp
--- a/Engine/Source/Runtime/Online/HTTP/Private/HttpThread.cpp
+++ b/Engine/Source/Runtime/Online/HTTP/Private/HttpThread.cpp
@@ -5,6 +5,7 @@
 #include "HAL/PlatformTime.h"
 #include "HAL/PlatformProcess.h"
 #include "HAL/RunnableThread.h"
+#include "HttpManager.h"
 #include "Misc/ConfigCacheIni.h"
 #include "Misc/CommandLine.h"
 #include "Misc/Fork.h"
@@ -296,6 +297,7 @@ void FHttpThreadBase::Process(TArray<IHttpThreadedRequest*>& RequestsToCancel, T
 			if (Request->GetDelegateThreadPolicy() == EHttpRequestDelegateThreadPolicy::CompleteOnHttpThread)
 			{
 				Request->FinishRequest();
+				FHttpModule::Get().GetHttpManager().BroadcastHttpRequestCompleted(Request->AsShared());
 			}
 
 			CompletedThreadedRequests.Enqueue(Request);
diff --git a/Engine/Source/Runtime/Online/HTTP/Public/HttpManager.h b/Engine/Source/Runtime/Online/HTTP/Public/HttpManager.h
--- a/Engine/Source/Runtime/Online/HTTP/Public/HttpManager.h
+++ b/Engine/Source/Runtime/Online/HTTP/Public/HttpManager.h
@@ -87,7 +87,18 @@ public:
 	 */
 	HTTP_API void AddRequest(const FHttpRequestRef& Request);
 
+	/**
+	 * Set a delegate to be triggered when an http request added to http manager. 
+	 * NOTE: The delegate can be triggered from different threads, depends on which
+	 * thread the request created. So make sure the delegate set here is thread-safe.
+	 */
 	HTTP_API void SetRequestAddedDelegate(const FHttpManagerRequestAddedDelegate& Delegate);
+
+	/**
+	 * Set a delegate to be triggered when an http request completed. 
+	 * NOTE: The delegate can be triggered from different threads, it depends on the delegate policy set 
+	 * to each http request. So make sure the delegate set here is thread-safe.
+	 */
 	HTTP_API void SetRequestCompletedDelegate(const FHttpManagerRequestCompletedDelegate& Delegate);
 
 	/**
diff --git a/Engine/Source/Runtime/Online/HTTP/Public/Interfaces/IHttpRequest.h b/Engine/Source/Runtime/Online/HTTP/Public/Interfaces/IHttpRequest.h
--- a/Engine/Source/Runtime/Online/HTTP/Public/Interfaces/IHttpRequest.h
+++ b/Engine/Source/Runtime/Online/HTTP/Public/Interfaces/IHttpRequest.h
@@ -387,7 +387,14 @@ public:
 	virtual float GetElapsedTime() const = 0;
 
 	/**
-	 * Set thread policy about which thread to complete this request
+	 * Set thread policy about which thread to trigger the delegates, set by FHttpManager::SetRequestCompletedDelegate, 
+	 * IHttpRequest::OnHeaderReceived, IHttpRequest::OnRequestProgress64 and IHttpRequest::OnProcessRequestComplete.
+	 *
+	 * Note that when set it as CompleteOnHttpThread, the thread to trigger delegates could be any thread 
+	 * depends on the implementation. User code should make the delegate thread-safe and shouldn't assume 
+	 * it's triggered by the thread where this request get created.
+	 * 
+	 * @param InThreadPolicy - The thread policy to indicate which thread to trigger the delegates
 	 */
 	virtual void SetDelegateThreadPolicy(EHttpRequestDelegateThreadPolicy InThreadPolicy) = 0;
 
