Fix the warnings in WebTests project:
- Set http log level to error if the warning logs are expected for specified tests;
- There is a warning about failing to delete file downloaded, because the file is still in use, close the file reader which verifies the file;
- There is a warning about failing to delete file uploaded, because the file is still in use. Now http upload request closes the file reader when complete.

diff --git a/Engine/Source/Programs/WebTests/Private/TestHttp.cpp b/Engine/Source/Programs/WebTests/Private/TestHttp.cpp
--- a/Engine/Source/Programs/WebTests/Private/TestHttp.cpp
+++ b/Engine/Source/Programs/WebTests/Private/TestHttp.cpp
@@ -35,6 +35,7 @@ public:
 		: WebServerIp(TEXT("127.0.0.1"))
 		, WebServerHttpPort(8000)
 		, bRunHeavyTests(false)
+		, OldVerbosity(LogHttp.GetVerbosity())
 	{
 		ParseSettingsFromCommandLine();
 
@@ -48,6 +49,11 @@ public:
 		IModuleInterface* Module = HttpModule;
 		Module->ShutdownModule();
 		delete Module;
+
+		if (OldVerbosity != LogHttp.GetVerbosity())
+		{
+			LogHttp.SetVerbosity(OldVerbosity);
+		}
 	}
 
 	void ParseSettingsFromCommandLine()
@@ -56,6 +62,11 @@ public:
 		FParse::Bool(FCommandLine::Get(), TEXT("run_heavy_tests"), bRunHeavyTests);
 	}
 
+	void DisableWarningsInThisTest()
+	{
+		LogHttp.SetVerbosity(ELogVerbosity::Error);
+	}
+
 	const FString UrlWithInvalidPortToTestConnectTimeout() const { return FString::Format(TEXT("http://{0}:{1}"), { *WebServerIp, 8765 }); }
 	const FString UrlBase() const { return FString::Format(TEXT("http://{0}:{1}"), { *WebServerIp, WebServerHttpPort }); }
 	const FString UrlHttpTests() const { return FString::Format(TEXT("{0}/webtests/httptests"), { *UrlBase() }); }
@@ -67,10 +78,13 @@ public:
 	FHttpModule* HttpModule;
 
 	bool bRunHeavyTests;
+	ELogVerbosity::Type OldVerbosity;
 };
 
 TEST_CASE_METHOD(FHttpModuleTestFixture, "Shutdown http module without issue when there are ongoing http requests.", HTTP_TAG)
 {
+	DisableWarningsInThisTest();
+
 	uint32 ChunkSize = 1024 * 1024;
 	TArray<uint8> DataChunk;
 	DataChunk.SetNum(ChunkSize);
@@ -322,6 +336,8 @@ TEST_CASE_METHOD(FWaitUntilCompleteHttpFixture, "Streaming http download", HTTP_
 	}
 	SECTION("Use customized stream to receive response body but failed when serialize")
 	{
+		DisableWarningsInThisTest();
+
 		class FTestHttpReceiveStream final : public FArchive
 		{
 		public:
@@ -350,6 +366,8 @@ TEST_CASE_METHOD(FWaitUntilCompleteHttpFixture, "Streaming http download", HTTP_
 	}
 	SECTION("Use customized stream delegate to receive response body but failed when call")
 	{
+		DisableWarningsInThisTest();
+
 		FHttpRequestStreamDelegate Delegate;
 		Delegate.BindLambda([TotalBytesReceived](void* Ptr, int64 Length) {
 			*TotalBytesReceived += Length;
@@ -384,6 +402,7 @@ TEST_CASE_METHOD(FWaitUntilCompleteHttpFixture, "Streaming http download", HTTP_
 
 			TSharedRef<FArchive> FileToRead = MakeShareable(IFileManager::Get().CreateFileReader(*Filename));
 			CHECK(FileToRead->TotalSize() == Chunks * ChunkSize);
+			FileToRead->Close();
 
 			IFileManager::Get().Delete(*Filename);
 		});
@@ -734,6 +753,8 @@ void SetupURLRequestFilter(FHttpModule* HttpModule)
 
 TEST_CASE_METHOD(FWaitUntilCompleteHttpFixture, "Http request pre check will fail", HTTP_TAG)
 {
+	DisableWarningsInThisTest();
+
 	// Pre check will fail when domain is not allowed
 	UE::TestHttp::SetupURLRequestFilter(HttpModule);
 
@@ -762,6 +783,8 @@ TEST_CASE_METHOD(FWaitUntilCompleteHttpFixture, "Http request pre check will fai
 
 TEST_CASE_METHOD(FWaitThreadedHttpFixture, "Threaded http request pre check will fail", HTTP_TAG)
 {
+	DisableWarningsInThisTest();
+
 	ThreadedHttpRunnable.OnRunFromThread().BindLambda([this]() {
 		// Pre check will fail when domain is not allowed
 		UE::TestHttp::SetupURLRequestFilter(HttpModule);
@@ -808,6 +831,7 @@ TEST_CASE_METHOD(FWaitUntilCompleteHttpFixture, "Cancel http request connect bef
 	});
 	HttpRequest->ProcessRequest();
 	FPlatformProcess::Sleep(0.5);
+	DisableWarningsInThisTest();
 	HttpRequest->CancelRequest();
 }
 
diff --git a/Engine/Source/Runtime/Online/HTTP/Private/Curl/CurlHttp.cpp b/Engine/Source/Runtime/Online/HTTP/Private/Curl/CurlHttp.cpp
--- a/Engine/Source/Runtime/Online/HTTP/Private/Curl/CurlHttp.cpp
+++ b/Engine/Source/Runtime/Online/HTTP/Private/Curl/CurlHttp.cpp
@@ -361,7 +361,7 @@ bool FCurlHttpRequest::SetContentAsStreamedFile(const FString& Filename)
 	FArchive* File = IFileManager::Get().CreateFileReader(*Filename);
 	if (File)
 	{
-		RequestPayload = MakeUnique<FRequestPayloadInFileStream>(MakeShareable(File));
+		RequestPayload = MakeUnique<FRequestPayloadInFileStream>(MakeShareable(File), true/*bInCloseWhenComplete*/);
 	}
 	else
 	{
@@ -1183,6 +1183,11 @@ void FCurlHttpRequest::FinishRequest()
 {
 	QUICK_SCOPE_CYCLE_COUNTER(STAT_FCurlHttpRequest_FinishRequest);
 
+	if (RequestPayload.IsValid())
+	{
+		RequestPayload->Close();
+	}
+
 	curl_easy_setopt(EasyHandle, CURLOPT_SHARE, nullptr);
 	
 	CheckProgressDelegate();
diff --git a/Engine/Source/Runtime/Online/HTTP/Private/GenericPlatform/HttpRequestPayload.cpp b/Engine/Source/Runtime/Online/HTTP/Private/GenericPlatform/HttpRequestPayload.cpp
--- a/Engine/Source/Runtime/Online/HTTP/Private/GenericPlatform/HttpRequestPayload.cpp
+++ b/Engine/Source/Runtime/Online/HTTP/Private/GenericPlatform/HttpRequestPayload.cpp
@@ -33,7 +33,9 @@ bool FGenericPlatformHttp::IsURLEncoded(const TArray<uint8>& Payload)
 	return true;
 }
 
-FRequestPayloadInFileStream::FRequestPayloadInFileStream(TSharedRef<FArchive, ESPMode::ThreadSafe> InFile) : File(InFile)
+FRequestPayloadInFileStream::FRequestPayloadInFileStream(TSharedRef<FArchive, ESPMode::ThreadSafe> InFile, bool bInCloseWhenComplete) 
+	: File(InFile)
+	, bCloseWhenComplete(bInCloseWhenComplete)
 {
 }
 
@@ -82,6 +84,14 @@ size_t FRequestPayloadInFileStream::FillOutputBuffer(TArrayView<uint8> OutputBuf
 	return SizeToSendThisTime;
 }
 
+void FRequestPayloadInFileStream::Close()
+{
+	if (bCloseWhenComplete)
+	{
+		File->Close();
+	}
+}
+
 FRequestPayloadInMemory::FRequestPayloadInMemory(const TArray<uint8>& Array) : Buffer(Array)
 {
 }
@@ -126,3 +136,7 @@ size_t FRequestPayloadInMemory::FillOutputBuffer(TArrayView<uint8> OutputBuffer,
 	}
 	return SizeToSendThisTime;
 }
+
+void FRequestPayloadInMemory::Close()
+{
+}
diff --git a/Engine/Source/Runtime/Online/HTTP/Public/GenericPlatform/HttpRequestPayload.h b/Engine/Source/Runtime/Online/HTTP/Public/GenericPlatform/HttpRequestPayload.h
--- a/Engine/Source/Runtime/Online/HTTP/Public/GenericPlatform/HttpRequestPayload.h
+++ b/Engine/Source/Runtime/Online/HTTP/Public/GenericPlatform/HttpRequestPayload.h
@@ -40,20 +40,25 @@ public:
 	 * @return Returns the number of bytes copied into OutputBuffer
 	 */
 	virtual size_t FillOutputBuffer(TArrayView<uint8> OutputBuffer, size_t SizeAlreadySent) = 0;
+
+	/** Close the upload payload when the http request don't need to use it anymore */
+	virtual void Close() = 0;
 };
 
 class FRequestPayloadInFileStream : public FRequestPayload
 {
 public:
-	FRequestPayloadInFileStream(TSharedRef<FArchive, ESPMode::ThreadSafe> InFile);
+	FRequestPayloadInFileStream(TSharedRef<FArchive, ESPMode::ThreadSafe> InFile, bool bInCloseWhenComplete = false);
 	virtual ~FRequestPayloadInFileStream();
 	virtual uint64 GetContentLength() const override;
 	virtual const TArray<uint8>& GetContent() const override;
 	virtual bool IsURLEncoded() const override;
 	virtual size_t FillOutputBuffer(void* OutputBuffer, size_t MaxOutputBufferSize, size_t SizeAlreadySent) override;
 	virtual size_t FillOutputBuffer(TArrayView<uint8> OutputBuffer, size_t SizeAlreadySent) override;
+	virtual void Close() override;
 private:
 	TSharedRef<FArchive, ESPMode::ThreadSafe> File;
+	bool bCloseWhenComplete = false;
 };
 
 class FRequestPayloadInMemory : public FRequestPayload
@@ -67,6 +72,7 @@ public:
 	virtual bool IsURLEncoded() const override;
 	virtual size_t FillOutputBuffer(void* OutputBuffer, size_t MaxOutputBufferSize, size_t SizeAlreadySent) override;
 	virtual size_t FillOutputBuffer(TArrayView<uint8> OutputBuffer, size_t SizeAlreadySent) override;
+	virtual void Close() override;
 private:
 	TArray<uint8> Buffer;
 };
