Fix the warning in new test case when validate retry manager thread safety;
Fix the warning when the request didn't get removed from http manager when test.

diff --git a/Engine/Source/Programs/WebTests/Private/TestHttp.cpp b/Engine/Source/Programs/WebTests/Private/TestHttp.cpp
--- a/Engine/Source/Programs/WebTests/Private/TestHttp.cpp
+++ b/Engine/Source/Programs/WebTests/Private/TestHttp.cpp
@@ -193,6 +193,10 @@ public:
 			HttpModule->GetHttpManager().Tick(TickFrequency);
 			FPlatformProcess::Sleep(TickFrequency);
 		}
+
+		// In case in http thread the http request complete and set OngoingRequests to 0, http manager never 
+		// had chance to Tick and remove the request
+		HttpModule->GetHttpManager().Tick(TickFrequency);
 	}
 
 	std::atomic<int32> OngoingRequests = 0;
@@ -921,6 +925,8 @@ public:
 
 TEST_CASE_METHOD(FRetryHttpManagerThreadedRequestsFixture, "Retry manager is thread safe", HTTP_TAG)
 {
+	DisableWarningsInThisTest();
+
 	ThreadedHttpRunnable.OnRunFromThread().BindLambda([this]() {
 		LaunchDownloadRequests();
 		HttpRetryManager->BlockUntilFlushed(5.0);
