Enable PS5 for WebTests;
Fix WebTests warnings on PS5.

diff --git a/Engine/Build/LowLevelTests/Web.xml b/Engine/Build/LowLevelTests/Web.xml
--- a/Engine/Build/LowLevelTests/Web.xml
+++ b/Engine/Build/LowLevelTests/Web.xml
@@ -13,23 +13,26 @@
   <Property Name="WebHasAfterSteps" Value="True" />
   <Property Name="WebUsesCatch2" Value="True" />
   <Option Name="RunWebTests" DefaultValue="False" Description="Run Web Tests" />
-  <Property Name="WebSupportedPlatforms" Value="Win64" />
+  <Property Name="WebSupportedPlatforms" Value="Win64;PS5" />
   <Property Name="Win64WebTags" Value="" />
-  <Property Name="MacWebTags" Value="" />
+  <Property Name="PS5WebTags" Value="" />
+<Property Name="MacWebTags" Value="" />
   <Property Name="IOSWebTags" Value="" />
   <Property Name="AndroidWebTags" Value="" />
   <Property Name="LinuxWebTags" Value="" />
   <Property Name="LinuxArm64WebTags" Value="" />
   <Property Name="TVOSWebTags" Value="" />
   <Property Name="Win64WebExtraCompilationArgs" Value="" />
-  <Property Name="MacWebExtraCompilationArgs" Value="" />
+  <Property Name="PS5WebExtraCompilationArgs" Value="" />
+<Property Name="MacWebExtraCompilationArgs" Value="" />
   <Property Name="IOSWebExtraCompilationArgs" Value="" />
   <Property Name="AndroidWebExtraCompilationArgs" Value="" />
   <Property Name="LinuxWebExtraCompilationArgs" Value="" />
   <Property Name="LinuxArm64WebExtraCompilationArgs" Value="" />
   <Property Name="TVOSWebExtraCompilationArgs" Value="" />
   <Property Name="Win64WebRunSupported" Value="True" />
-  <Property Name="MacWebRunSupported" Value="True" />
+<Property Name="PS5WebRunSupported" Value="True" />
+<Property Name="MacWebRunSupported" Value="True" />
   <Property Name="IOSWebRunSupported" Value="False" />
   <Property Name="AndroidWebRunSupported" Value="False" />
   <Property Name="LinuxWebRunSupported" Value="True" />
diff --git a/Engine/Source/Programs/WebTests/Private/TestWebSockets.cpp b/Engine/Source/Programs/WebTests/Private/TestWebSockets.cpp
--- a/Engine/Source/Programs/WebTests/Private/TestWebSockets.cpp
+++ b/Engine/Source/Programs/WebTests/Private/TestWebSockets.cpp
@@ -4,6 +4,7 @@
 #include "HAL/PlatformProcess.h"
 #include "Http.h"
 #include "Misc/CommandLine.h"
+#include "WebSocketsLog.h"
 #include "WebSocketsModule.h"
 #include "IWebSocket.h"
 #include "TestHarness.h"
@@ -28,6 +29,7 @@ public:
 	FWebSocketsModuleTestFixture()
 		: WebServerIp(TEXT("127.0.0.1"))
 		, WebServerWebSocketsPort(8000)
+		, OldVerbosity(LogWebSockets.GetVerbosity())
 	{
 		ParseSettingsFromCommandLine();
 
@@ -50,11 +52,22 @@ public:
 		IModuleInterface* HttpModuleInterface = HttpModule;
 		HttpModuleInterface->ShutdownModule();
 		delete HttpModuleInterface;
+
+		if (OldVerbosity != LogWebSockets.GetVerbosity())
+		{
+			LogWebSockets.SetVerbosity(OldVerbosity);
+		}
 	}
 
 	void ParseSettingsFromCommandLine()
 	{
 		FParse::Value(FCommandLine::Get(), TEXT("web_server_ip"), WebServerIp);
+		FParse::Value(FCommandLine::Get(), TEXT("web_server_websockets_port"), WebServerWebSocketsPort);
+	}
+
+	void DisableWarningsInThisTest()
+	{
+		LogWebSockets.SetVerbosity(ELogVerbosity::Error);
 	}
 
 	const FString UrlWithInvalidPortToTestConnectTimeout() const { return FString::Format(TEXT("ws://{0}:{1}"), { *WebServerIp, 8765 }); }
@@ -65,6 +78,7 @@ public:
 	uint32 WebServerWebSocketsPort;
 	FWebSocketsModule* WebSocketsModule;
 	FHttpModule* HttpModule;
+	ELogVerbosity::Type OldVerbosity;
 };
 
 class FRunUntilQuitRequestedFixture : public FWebSocketsModuleTestFixture
@@ -121,3 +135,20 @@ TEST_CASE_METHOD(FRunUntilQuitRequestedFixture, "WebSockets can connect then sen
 	WebSocket->Connect();
 }
 
+TEST_CASE_METHOD(FRunUntilQuitRequestedFixture, "WebSockets module can shut down when there are still websockets connections", WEBSOCKETS_TAG)
+{
+	DisableWarningsInThisTest();
+
+	TSharedPtr<IWebSocket> WebSocket = WebSocketsModule->CreateWebSocket(FString::Format(TEXT("{0}/echo/"), { *UrlWebSocketsTests() }));
+
+	WebSocket->OnConnected().AddLambda([this, WebSocket](){
+		bQuitRequested = true;
+	});
+
+	WebSocket->OnConnectionError().AddLambda([this, WebSocket](const FString& /* Error */){
+		CHECK(false);
+		bQuitRequested = true;
+	});
+
+	WebSocket->Connect();
+}
diff --git a/Engine/Source/Runtime/Online/WebSockets/Private/Lws/LwsWebSocketsManager.cpp b/Engine/Source/Runtime/Online/WebSockets/Private/Lws/LwsWebSocketsManager.cpp
--- a/Engine/Source/Runtime/Online/WebSockets/Private/Lws/LwsWebSocketsManager.cpp
+++ b/Engine/Source/Runtime/Online/WebSockets/Private/Lws/LwsWebSocketsManager.cpp
@@ -193,6 +193,11 @@ void FLwsWebSocketsManager::ShutdownWebSockets()
 
 	if (LwsContext)
 	{
+		if (!Sockets.IsEmpty())
+		{
+			lws_cancel_service(LwsContext);
+		}
+
 		lws_context_destroy(LwsContext);
 		LwsContext = nullptr;
 	}
@@ -246,8 +251,6 @@ uint32 FLwsWebSocketsManager::Run()
 void FLwsWebSocketsManager::Stop()
 {
 	ExitRequest.Set(true);
-	// Safe to call from other threads
-	lws_cancel_service(LwsContext);
 }
 
 void FLwsWebSocketsManager::Exit()
diff --git a/Engine/Source/Runtime/Online/WebSockets/Private/WebSocketsLog.h b/Engine/Source/Runtime/Online/WebSockets/Private/WebSocketsLog.h
deleted file mode 100644
diff --git a/Engine/Source/Runtime/Online/WebSockets/Public/WebSocketsLog.h b/Engine/Source/Runtime/Online/WebSockets/Public/WebSocketsLog.h
new file mode 100644
--- /dev/null
+++ b/Engine/Source/Runtime/Online/WebSockets/Public/WebSocketsLog.h
@@ -0,0 +1,7 @@
+// Copyright Epic Games, Inc. All Rights Reserved.
+#pragma once
+
+#include "CoreMinimal.h"
+
+DECLARE_LOG_CATEGORY_EXTERN(LogWebSockets, Warning, All);
+
