Remove dependency of http request from http response;
Make http requests use GET by default if verb is not set.

diff --git a/Engine/Source/Programs/WebTests/Private/TestHttp.cpp b/Engine/Source/Programs/WebTests/Private/TestHttp.cpp
--- a/Engine/Source/Programs/WebTests/Private/TestHttp.cpp
+++ b/Engine/Source/Programs/WebTests/Private/TestHttp.cpp
@@ -10,6 +10,7 @@
 #include "Http.h"
 #include "Misc/CommandLine.h"
 #include "TestHarness.h"
+#include "Serialization/JsonSerializerMacros.h"
 
 /**
  *  HTTP Tests
@@ -210,8 +211,13 @@ public:
 TEST_CASE_METHOD(FWaitUntilCompleteHttpFixture, "Http Methods", HTTP_TAG)
 {
 	TSharedRef<IHttpRequest> HttpRequest = HttpModule->CreateRequest();
+	CHECK(HttpRequest->GetVerb() == TEXT("GET"));
+
 	HttpRequest->SetURL(UrlToTestMethods());
 
+	SECTION("Default GET")
+	{
+	}
 	SECTION("GET")
 	{
 		HttpRequest->SetVerb(TEXT("GET"));
@@ -622,6 +628,71 @@ TEST_CASE_METHOD(FWaitUntilCompleteHttpFixture, "Redirect enabled by default and
 	HttpRequest->ProcessRequest();
 }
 
+// Response shared ptr should be able to be kept by user code and valid to access without http request
+class FValidateResponseDependencyFixture : public FWaitUntilCompleteHttpFixture
+{
+public:
+	DECLARE_DELEGATE(FValidateResponseDependencyDelegate);
+
+	~FValidateResponseDependencyFixture()
+	{
+		WaitUntilAllHttpRequestsComplete();
+
+		ValidateResponseDependencyDelegate.ExecuteIfBound();
+	}
+
+	FValidateResponseDependencyDelegate ValidateResponseDependencyDelegate;
+};
+
+TEST_CASE_METHOD(FValidateResponseDependencyFixture, "Http query with parameters", HTTP_TAG)
+{
+	struct FQueryWithParamsResponse : public FJsonSerializable
+	{
+		int32 VarInt;
+		FString VarStr;
+
+		BEGIN_JSON_SERIALIZER
+			JSON_SERIALIZE("var_int", VarInt);
+			JSON_SERIALIZE("var_str", VarStr);
+		END_JSON_SERIALIZER
+	};
+
+	TSharedRef<IHttpRequest> HttpRequest = HttpModule->CreateRequest();
+	FString UrlQueryWithParams = FString::Format(TEXT("{0}/query_with_params/?var_int=3&var_str=abc"), { *UrlHttpTests() });
+	HttpRequest->SetURL(UrlQueryWithParams);
+	HttpRequest->SetVerb(TEXT("GET"));
+	HttpRequest->OnProcessRequestComplete().BindLambda([this, UrlQueryWithParams](FHttpRequestPtr HttpRequest, FHttpResponsePtr HttpResponse, bool bSucceeded) {
+		CHECK(bSucceeded);
+		REQUIRE(HttpResponse != nullptr);
+		CHECK(HttpResponse->GetResponseCode() == 200);
+
+		CHECK(HttpRequest->GetURL() == UrlQueryWithParams);
+
+		FQueryWithParamsResponse QueryWithParamsResponse;
+		REQUIRE(QueryWithParamsResponse.FromJson(HttpResponse->GetContentAsString()));
+
+		CHECK(FString::FromInt(QueryWithParamsResponse.VarInt) == HttpRequest->GetURLParameter(TEXT("var_int")));
+		CHECK(QueryWithParamsResponse.VarStr == HttpRequest->GetURLParameter(TEXT("var_str")));
+
+		CHECK(FString::FromInt(QueryWithParamsResponse.VarInt) == HttpResponse->GetURLParameter(TEXT("var_int")));
+		CHECK(QueryWithParamsResponse.VarStr == HttpResponse->GetURLParameter(TEXT("var_str")));
+
+		ValidateResponseDependencyDelegate.BindLambda([HttpResponse, UrlQueryWithParams, QueryWithParamsResponse](){
+			// Validate all interfaces of http response can be called without accessing the destroyed http request
+			CHECK(HttpResponse->GetResponseCode() == 200);
+			CHECK(!HttpResponse->GetContent().IsEmpty());
+			CHECK(!HttpResponse->GetContentAsString().IsEmpty());
+			CHECK(HttpResponse->GetContentType() == TEXT("application/json"));
+			CHECK(HttpResponse->GetHeader("Content-Type") == TEXT("application/json"));
+			CHECK(!HttpResponse->GetAllHeaders().IsEmpty());
+			CHECK(HttpResponse->GetURL() == UrlQueryWithParams);
+			CHECK(HttpResponse->GetURLParameter(TEXT("var_int")) == FString::FromInt(QueryWithParamsResponse.VarInt));
+			CHECK(HttpResponse->GetURLParameter(TEXT("var_str")) == QueryWithParamsResponse.VarStr);
+		});
+	});
+	HttpRequest->ProcessRequest();
+}
+
 class FThreadedHttpRunnable : public FRunnable
 {
 public:
@@ -742,6 +813,34 @@ TEST_CASE_METHOD(FWaitThreadedHttpFixture, "Http download request progress callb
 	CHECK(bRequestProgressTriggered);
 }
 
+TEST_CASE_METHOD(FWaitUntilCompleteHttpFixture, "Http request pre check will fail", HTTP_TAG)
+{
+	DisableWarningsInThisTest();
+
+	TSharedRef<IHttpRequest> HttpRequest = HttpModule->CreateRequest();
+
+	SECTION("when verb was set to empty")
+	{
+		HttpRequest->SetURL(UrlToTestMethods());
+		HttpRequest->SetVerb(TEXT(""));
+	}
+	SECTION("when url protocol is not valid")
+	{
+		HttpRequest->SetURL("http_abc://www.epicgames.com");
+		HttpRequest->SetVerb(TEXT("GET"));
+	}
+	SECTION("when url was not set")
+	{
+		HttpRequest->SetVerb(TEXT("GET"));
+	}
+
+	HttpRequest->OnProcessRequestComplete().BindLambda([](FHttpRequestPtr HttpRequest, FHttpResponsePtr HttpResponse, bool bSucceeded) {
+		CHECK(!bSucceeded);
+	});
+
+	HttpRequest->ProcessRequest();
+}
+
 namespace UE
 {
 namespace TestHttp
@@ -759,7 +858,7 @@ void SetupURLRequestFilter(FHttpModule* HttpModule)
 }
 }
 
-TEST_CASE_METHOD(FWaitUntilCompleteHttpFixture, "Http request pre check will fail", HTTP_TAG)
+TEST_CASE_METHOD(FWaitUntilCompleteHttpFixture, "Http request pre check will fail by thread policy", HTTP_TAG)
 {
 	DisableWarningsInThisTest();
 
@@ -789,7 +888,7 @@ TEST_CASE_METHOD(FWaitUntilCompleteHttpFixture, "Http request pre check will fai
 	HttpRequest->ProcessRequest();
 }
 
-TEST_CASE_METHOD(FWaitThreadedHttpFixture, "Threaded http request pre check will fail", HTTP_TAG)
+TEST_CASE_METHOD(FWaitThreadedHttpFixture, "Threaded http request pre check will fail by thread policy", HTTP_TAG)
 {
 	DisableWarningsInThisTest();
 
diff --git a/Engine/Source/Programs/WebTestsServer/httptests/urls.py b/Engine/Source/Programs/WebTestsServer/httptests/urls.py
--- a/Engine/Source/Programs/WebTestsServer/httptests/urls.py
+++ b/Engine/Source/Programs/WebTestsServer/httptests/urls.py
@@ -4,6 +4,7 @@ from . import views
 
 urlpatterns = [
     path("httptests/methods", views.methods, name="methods"),
+    path("httptests/query_with_params/", views.query_with_params, name="query_with_params"),
     path("httptests/get_large_response_without_chunks/<int:bytes_number>/", views.get_large_response_without_chunks, name="get_large_response_without_chunks"),
     path("httptests/nonstreaming_receivetimeout/<int:wait_time>/", views.nonstreaming_receivetimeout, name="nonstreaming_receivetimeout"),
     path("httptests/streaming_download/<int:chunks>/<int:chunk_size>/", views.streaming_download, name="streaming_download"),
diff --git a/Engine/Source/Programs/WebTestsServer/httptests/views.py b/Engine/Source/Programs/WebTestsServer/httptests/views.py
--- a/Engine/Source/Programs/WebTestsServer/httptests/views.py
+++ b/Engine/Source/Programs/WebTestsServer/httptests/views.py
@@ -13,6 +13,12 @@ logger = logging.getLogger('django')
 def methods(request):
     return JsonResponse({})
 
+@api_view(['GET'])
+def query_with_params(request):
+    var_int = int(request.GET.get('var_int'))
+    var_str = request.GET.get('var_str')
+    return JsonResponse({'var_int' : var_int, 'var_str' : var_str}, content_type="application/json")
+
 @api_view(['GET'])
 def get_large_response_without_chunks(request, bytes_number):
     data = "d" * bytes_number
diff --git a/Engine/Source/Runtime/Online/HTTP/Private/Apple/AppleHTTPNSUrlConnection.cpp b/Engine/Source/Runtime/Online/HTTP/Private/Apple/AppleHTTPNSUrlConnection.cpp
--- a/Engine/Source/Runtime/Online/HTTP/Private/Apple/AppleHTTPNSUrlConnection.cpp
+++ b/Engine/Source/Runtime/Online/HTTP/Private/Apple/AppleHTTPNSUrlConnection.cpp
@@ -86,33 +86,6 @@ void FAppleHttpNSUrlConnectionRequest::SetURL(const FString& URL)
 	Request.URL = [NSURL URLWithString: URL.GetNSString()];
 }
 
-
-FString FAppleHttpNSUrlConnectionRequest::GetURLParameter(const FString& ParameterName) const
-{
-	SCOPED_AUTORELEASE_POOL;
-	UE_LOG(LogHttp, Verbose, TEXT("FAppleHttpNSUrlConnectionRequest::GetURLParameter() - %s"), *ParameterName);
-
-	NSRange ParametersStart = [Request.URL.query rangeOfString:@"?"];
-	if (ParametersStart.location != NSNotFound && ParametersStart.length > 0)
-	{
-		NSString* ParametersStr = [Request.URL.query substringFromIndex:ParametersStart.location + 1];
-		NSString* ParameterNameStr = ParameterName.GetNSString();
-		NSArray* Parameters = [ParametersStr componentsSeparatedByString:@"&"];
-		for (NSString* Parameter in Parameters)
-		{
-			NSArray* KeyValue = [Parameter componentsSeparatedByString:@"="];
-			NSString* Key = KeyValue[0];
-			if ([Key compare:ParameterNameStr] == NSOrderedSame)
-			{
-				return FString(KeyValue[1]);
-			}
-		}
-	}
-
-	return FString();
-}
-
-
 FString FAppleHttpNSUrlConnectionRequest::GetHeader(const FString& HeaderName) const
 {
 	SCOPED_AUTORELEASE_POOL;
@@ -846,13 +819,12 @@ float FAppleHttpNSUrlConnectionRequest::GetElapsedTime() const
  **************************************************************************/
 
 FAppleHttpNSUrlConnectionResponse::FAppleHttpNSUrlConnectionResponse(const FAppleHttpNSUrlConnectionRequest& InRequest)
-	: Request( InRequest )
+	: FHttpResponseCommon(InRequest.GetURL())
 {
 	UE_LOG(LogHttp, Verbose, TEXT("FAppleHttpNSUrlConnectionResponse::FAppleHttpNSUrlConnectionResponse()"));
 	ResponseWrapper = [[FHttpResponseAppleNSUrlConnectionWrapper alloc] init];
 }
 
-
 FAppleHttpNSUrlConnectionResponse::~FAppleHttpNSUrlConnectionResponse()
 {
 	UE_LOG(LogHttp, Verbose, TEXT("FAppleHttpNSUrlConnectionResponse::~FAppleHttpNSUrlConnectionResponse()"));
@@ -862,34 +834,6 @@ FAppleHttpNSUrlConnectionResponse::~FAppleHttpNSUrlConnectionResponse()
 	ResponseWrapper = nil;
 }
 
-
-FString FAppleHttpNSUrlConnectionResponse::GetURL() const
-{
-	UE_LOG(LogHttp, Verbose, TEXT("FAppleHttpNSUrlConnectionResponse::GetURL()"));
-	return FString(Request.Request.URL.query);
-}
-
-
-FString FAppleHttpNSUrlConnectionResponse::GetURLParameter(const FString& ParameterName) const
-{
-	SCOPED_AUTORELEASE_POOL;
-	UE_LOG(LogHttp, Verbose, TEXT("FAppleHttpNSUrlConnectionResponse::GetURLParameter()"));
-
-	NSString* ParameterNameStr = ParameterName.GetNSString();
-	NSArray* Parameters = [[[Request.Request URL] query] componentsSeparatedByString:@"&"];
-	for (NSString* Parameter in Parameters)
-	{
-		NSArray* KeyValue = [Parameter componentsSeparatedByString:@"="];
-		NSString* Key = [KeyValue objectAtIndex:0];
-		if ([Key compare:ParameterNameStr] == NSOrderedSame)
-		{
-			return FString([[KeyValue objectAtIndex:1] stringByRemovingPercentEncoding]);
-		}
-	}
-	return FString();
-}
-
-
 FString FAppleHttpNSUrlConnectionResponse::GetHeader(const FString& HeaderName) const
 {
 	SCOPED_AUTORELEASE_POOL;
@@ -937,7 +881,7 @@ const TArray<uint8>& FAppleHttpNSUrlConnectionResponse::GetContent() const
 {
 	if( !IsReady() )
 	{
-		UE_LOG(LogHttp, Warning, TEXT("Payload is incomplete. Response still processing. %p"), &Request);
+		UE_LOG(LogHttp, Warning, TEXT("Payload is incomplete. Response still processing. %s"), *GetURL());
 	}
 	else
 	{
diff --git a/Engine/Source/Runtime/Online/HTTP/Private/Apple/AppleHTTPNSUrlConnection.h b/Engine/Source/Runtime/Online/HTTP/Private/Apple/AppleHTTPNSUrlConnection.h
--- a/Engine/Source/Runtime/Online/HTTP/Private/Apple/AppleHTTPNSUrlConnection.h
+++ b/Engine/Source/Runtime/Online/HTTP/Private/Apple/AppleHTTPNSUrlConnection.h
@@ -5,7 +5,7 @@
 #include "CoreMinimal.h"
 #include "Containers/UnrealString.h"
 #include "GenericPlatform/HttpRequestCommon.h"
-#include "Interfaces/IHttpResponse.h"
+#include "GenericPlatform/HttpResponseCommon.h"
 #include "HttpManager.h"
 #include "PlatformHttp.h"
 
@@ -25,7 +25,6 @@ public:
 
 	//~ Begin IHttpBase Interface
 	virtual FString GetURL() const override;
-	virtual FString GetURLParameter(const FString& ParameterName) const override;
 	virtual FString GetHeader(const FString& HeaderName) const override;
 	virtual TArray<FString> GetAllHeaders() const override;	
 	virtual FString GetContentType() const override;
@@ -173,24 +172,18 @@ private:
 /**
  * Apple implementation of an Http response
  */
-class FAppleHttpNSUrlConnectionResponse : public IHttpResponse
+class FAppleHttpNSUrlConnectionResponse : public FHttpResponseCommon
 {
 private:
 	// This is the NSHTTPURLResponse, all our functionality will deal with.
 	FHttpResponseAppleNSUrlConnectionWrapper* ResponseWrapper;
 
-	/** Request that owns this response */
-	const FAppleHttpNSUrlConnectionRequest& Request;
-
-
 public:
 	// implementation friends
 	friend class FAppleHttpNSUrlConnectionRequest;
 
 
 	//~ Begin IHttpBase Interface
-	virtual FString GetURL() const override;
-	virtual FString GetURLParameter(const FString& ParameterName) const override;
 	virtual FString GetHeader(const FString& HeaderName) const override;
 	virtual TArray<FString> GetAllHeaders() const override;	
 	virtual FString GetContentType() const override;
diff --git a/Engine/Source/Runtime/Online/HTTP/Private/Apple/AppleHTTPNSUrlSession.cpp b/Engine/Source/Runtime/Online/HTTP/Private/Apple/AppleHTTPNSUrlSession.cpp
--- a/Engine/Source/Runtime/Online/HTTP/Private/Apple/AppleHTTPNSUrlSession.cpp
+++ b/Engine/Source/Runtime/Online/HTTP/Private/Apple/AppleHTTPNSUrlSession.cpp
@@ -424,31 +424,6 @@ void FAppleHttpNSUrlSessionRequest::SetURL(const FString& URL)
 	Request.URL = [NSURL URLWithString: URL.GetNSString()];
 }
 
-FString FAppleHttpNSUrlSessionRequest::GetURLParameter(const FString& ParameterName) const
-{
-	SCOPED_AUTORELEASE_POOL;
-	UE_LOG(LogHttp, Verbose, TEXT("FAppleHttpNSUrlSessionRequest::GetURLParameter() - %s"), *ParameterName);
-
-	NSRange ParametersStart = [Request.URL.query rangeOfString:@"?"];
-	if (ParametersStart.location != NSNotFound && ParametersStart.length > 0)
-	{
-		NSString* ParametersStr = [Request.URL.query substringFromIndex:ParametersStart.location + 1];
-		NSString* ParameterNameStr = ParameterName.GetNSString();
-		NSArray* Parameters = [ParametersStr componentsSeparatedByString:@"&"];
-		for (NSString* Parameter in Parameters)
-		{
-			NSArray* KeyValue = [Parameter componentsSeparatedByString:@"="];
-			NSString* Key = KeyValue[0];
-			if ([Key compare:ParameterNameStr] == NSOrderedSame)
-			{
-				return FString(KeyValue[1]);
-			}
-		}
-	}
-
-	return FString();
-}
-
 FString FAppleHttpNSUrlSessionRequest::GetHeader(const FString& HeaderName) const
 {
 	SCOPED_AUTORELEASE_POOL;
@@ -866,7 +841,7 @@ void FAppleHttpNSUrlSessionRequest::TickThreadedRequest(float DeltaSeconds)
  **************************************************************************/
 
 FAppleHttpNSUrlSessionResponse::FAppleHttpNSUrlSessionResponse(const FAppleHttpNSUrlSessionRequest& InRequest)
-	: Request( InRequest )
+	: FHttpResponseCommon(InRequest.GetURL())
 {
 	UE_LOG(LogHttp, Verbose, TEXT("FAppleHttpNSUrlSessionResponse::FAppleHttpNSUrlSessionResponse()"));
 	ResponseDelegate = [[FAppleHttpNSUrlSessionResponseDelegate alloc] initWithResponseStream: InRequest.ResponseBodyReceiveStream];
@@ -890,31 +865,6 @@ void FAppleHttpNSUrlSessionResponse::CleanSharedObjects()
 	[ResponseDelegate ClearResponseStream];
 }
 
-FString FAppleHttpNSUrlSessionResponse::GetURL() const
-{
-	UE_LOG(LogHttp, Verbose, TEXT("FAppleHttpNSUrlSessionResponse::GetURL()"));
-	return FString(Request.Request.URL.query);
-}
-
-FString FAppleHttpNSUrlSessionResponse::GetURLParameter(const FString& ParameterName) const
-{
-	SCOPED_AUTORELEASE_POOL;
-	UE_LOG(LogHttp, Verbose, TEXT("FAppleHttpNSUrlSessionResponse::GetURLParameter()"));
-
-	NSString* ParameterNameStr = ParameterName.GetNSString();
-	NSArray* Parameters = [[[Request.Request URL] query] componentsSeparatedByString:@"&"];
-	for (NSString* Parameter in Parameters)
-	{
-		NSArray* KeyValue = [Parameter componentsSeparatedByString:@"="];
-		NSString* Key = [KeyValue objectAtIndex:0];
-		if ([Key compare:ParameterNameStr] == NSOrderedSame)
-		{
-			return FString([[KeyValue objectAtIndex:1] stringByRemovingPercentEncoding]);
-		}
-	}
-	return FString();
-}
-
 FString FAppleHttpNSUrlSessionResponse::GetHeader(const FString& HeaderName) const
 {
 	SCOPED_AUTORELEASE_POOL;
@@ -957,7 +907,7 @@ const TArray<uint8>& FAppleHttpNSUrlSessionResponse::GetContent() const
 {
 	if( !IsReady() )
 	{
-		UE_LOG(LogHttp, Warning, TEXT("Payload is incomplete. Response still processing. %p"), &Request);
+		UE_LOG(LogHttp, Warning, TEXT("Payload is incomplete. Response still processing. %s"), *GetURL());
 	}
 	else
 	{
diff --git a/Engine/Source/Runtime/Online/HTTP/Private/Apple/AppleHTTPNSUrlSession.h b/Engine/Source/Runtime/Online/HTTP/Private/Apple/AppleHTTPNSUrlSession.h
--- a/Engine/Source/Runtime/Online/HTTP/Private/Apple/AppleHTTPNSUrlSession.h
+++ b/Engine/Source/Runtime/Online/HTTP/Private/Apple/AppleHTTPNSUrlSession.h
@@ -3,8 +3,8 @@
 #pragma once
 
 #include "CoreMinimal.h"
+#include "GenericPlatform/HttpResponseCommon.h"
 #include "IHttpThreadedRequest.h"
-#include "Interfaces/IHttpResponse.h"
 #include "PlatformHttp.h"
 
 
@@ -26,7 +26,6 @@ public:
 
 	//~ Begin IHttpBase Interface
 	virtual FString GetURL() const override;
-	virtual FString GetURLParameter(const FString& ParameterName) const override;
 	virtual FString GetHeader(const FString& HeaderName) const override;
 	virtual TArray<FString> GetAllHeaders() const override;	
 	virtual FString GetContentType() const override;
@@ -134,24 +133,18 @@ private:
 /**
  * Apple implementation of an Http response
  */
-class FAppleHttpNSUrlSessionResponse : public IHttpResponse
+class FAppleHttpNSUrlSessionResponse : public FHttpResponseCommon
 {
 private:
 	// Delegate implementation. Keeps the response state and data
 	FAppleHttpNSUrlSessionResponseDelegate* ResponseDelegate;
 
-	/** Request that owns this response */
-	const FAppleHttpNSUrlSessionRequest& Request;
-
-
 public:
 	// implementation friends
 	friend class FAppleHttpNSUrlSessionRequest;
 
 
 	//~ Begin IHttpBase Interface
-	virtual FString GetURL() const override;
-	virtual FString GetURLParameter(const FString& ParameterName) const override;
 	virtual FString GetHeader(const FString& HeaderName) const override;
 	virtual TArray<FString> GetAllHeaders() const override;	
 	virtual FString GetContentType() const override;
diff --git a/Engine/Source/Runtime/Online/HTTP/Private/Curl/CurlHttp.cpp b/Engine/Source/Runtime/Online/HTTP/Private/Curl/CurlHttp.cpp
--- a/Engine/Source/Runtime/Online/HTTP/Private/Curl/CurlHttp.cpp
+++ b/Engine/Source/Runtime/Online/HTTP/Private/Curl/CurlHttp.cpp
@@ -62,21 +62,22 @@ static CURLcode sslctx_function(CURL * curl, void * sslctx, void * parm)
 #endif //#if WITH_SSL
 
 FCurlHttpRequest::FCurlHttpRequest()
-	:	EasyHandle(nullptr)
-	,	HeaderList(nullptr)
-	,	bCanceled(false)
-	,	bCurlRequestCompleted(false)
-	,	bRedirected(false)
-	,	CurlAddToMultiResult(CURLM_OK)
-	,	CurlCompletionResult(CURLE_OK)
-	,	ElapsedTime(0.0f)
-	,	TimeSinceLastResponse(0.0f)
-	,	bAnyHttpActivity(false)
-	,   BytesSent(0)
-	,	TotalBytesSent(0)
-	,	LastReportedBytesRead(0)
-	,	LastReportedBytesSent(0)
-	,   LeastRecentlyCachedInfoMessageIndex(0)
+	: EasyHandle(nullptr)
+	, HeaderList(nullptr)
+	, Verb(TEXT("GET"))
+	, bCanceled(false)
+	, bCurlRequestCompleted(false)
+	, bRedirected(false)
+	, CurlAddToMultiResult(CURLM_OK)
+	, CurlCompletionResult(CURLE_OK)
+	, ElapsedTime(0.0f)
+	, TimeSinceLastResponse(0.0f)
+	, bAnyHttpActivity(false)
+	, BytesSent(0)
+	, TotalBytesSent(0)
+	, LastReportedBytesRead(0)
+	, LastReportedBytesSent(0)
+	, LeastRecentlyCachedInfoMessageIndex(0)
 {
 	checkf(FCurlHttpManager::IsInit(), TEXT("Curl request was created while the library is shutdown"));
 
@@ -203,42 +204,6 @@ FString FCurlHttpRequest::GetURL() const
 	return URL;
 }
 
-FString FCurlHttpRequest::GetURLParameter(const FString& ParameterName) const
-{
-	TArray<FString> StringElements;
-
-	//Parameters start after "?" in url
-	FString Path, Parameters;
-	if (URL.Split(TEXT("?"), &Path, &Parameters))
-	{
-		int32 NumElems = Parameters.ParseIntoArray(StringElements, TEXT("&"), true);
-		check(NumElems == StringElements.Num());
-		
-		FString ParamValDelimiter(TEXT("="));
-		for (int Idx = 0; Idx < NumElems; ++Idx )
-		{
-			FString Param, Value;
-			if (StringElements[Idx].Split(ParamValDelimiter, &Param, &Value) && Param == ParameterName)
-			{
-				// unescape
-				auto Converter = StringCast<ANSICHAR>(*Value);
-				char * EscapedAnsi = (char *)Converter.Get();
-				int32 EscapedLength = Converter.Length();
-
-				int32 UnescapedLength = 0;	
-				char * UnescapedAnsi = curl_easy_unescape(EasyHandle, EscapedAnsi, EscapedLength, &UnescapedLength);
-				
-				FString UnescapedValue(ANSI_TO_TCHAR(UnescapedAnsi));
-				curl_free(UnescapedAnsi);
-				
-				return UnescapedValue;
-			}
-		}
-	}
-
-	return FString();
-}
-
 FString FCurlHttpRequest::GetHeader(const FString& HeaderName) const
 {
 	FString Result;
@@ -1358,28 +1323,14 @@ float FCurlHttpRequest::GetElapsedTime() const
 
 // FCurlHttpRequest
 
-FCurlHttpResponse::FCurlHttpResponse(FCurlHttpRequest& InRequest)
-	:	Request(InRequest)
-	,	TotalBytesRead(0)
-	,	HttpCode(EHttpResponseCodes::Unknown)
-	,	ContentLength(0)
-	,	bIsReady(0)
-	,	bSucceeded(0)
-{
-}
-
-FCurlHttpResponse::~FCurlHttpResponse()
-{	
-}
-
-FString FCurlHttpResponse::GetURL() const
-{
-	return Request.GetURL();
-}
-
-FString FCurlHttpResponse::GetURLParameter(const FString& ParameterName) const
+FCurlHttpResponse::FCurlHttpResponse(const FCurlHttpRequest& InRequest)
+	: FHttpResponseCommon(InRequest.GetURL())
+	, TotalBytesRead(0)
+	, HttpCode(EHttpResponseCodes::Unknown)
+	, ContentLength(0)
+	, bIsReady(0)
+	, bSucceeded(0)
 {
-	return Request.GetURLParameter(ParameterName);
 }
 
 FString FCurlHttpResponse::GetHeader(const FString& HeaderName) const
@@ -1387,8 +1338,7 @@ FString FCurlHttpResponse::GetHeader(const FString& HeaderName) const
 	FString Result;
 	if (!bIsReady)
 	{
-		UE_LOG(LogHttp, Warning, TEXT("Can't get cached header [%s]. Response still processing. %p"),
-			*HeaderName, &Request);
+		UE_LOG(LogHttp, Warning, TEXT("Can't get cached header [%s]. Response still processing. %s"), *HeaderName, *GetURL());
 	}
 	else
 	{
@@ -1406,7 +1356,7 @@ TArray<FString> FCurlHttpResponse::GetAllHeaders() const
 	TArray<FString> Result;
 	if (!bIsReady)
 	{
-		UE_LOG(LogHttp, Warning, TEXT("Can't get cached headers. Response still processing. %p"),&Request);
+		UE_LOG(LogHttp, Warning, TEXT("Can't get cached headers. Response still processing. %s"), *GetURL());
 	}
 	else
 	{
@@ -1433,7 +1383,7 @@ const TArray<uint8>& FCurlHttpResponse::GetContent() const
 {
 	if (!bIsReady)
 	{
-		UE_LOG(LogHttp, Warning, TEXT("Payload is incomplete. Response still processing. %p"),&Request);
+		UE_LOG(LogHttp, Warning, TEXT("Payload is incomplete. Response still processing. %s"), *GetURL());
 	}
 	return Payload;
 }
diff --git a/Engine/Source/Runtime/Online/HTTP/Private/Curl/CurlHttp.h b/Engine/Source/Runtime/Online/HTTP/Private/Curl/CurlHttp.h
--- a/Engine/Source/Runtime/Online/HTTP/Private/Curl/CurlHttp.h
+++ b/Engine/Source/Runtime/Online/HTTP/Private/Curl/CurlHttp.h
@@ -5,10 +5,10 @@
 #include "CoreMinimal.h"
 #include "HAL/LowLevelMemTracker.h"
 #include "HAL/ThreadSafeCounter.h"
-#include "Interfaces/IHttpResponse.h"
 #include "IHttpThreadedRequest.h"
 #include "Containers/SpscQueue.h"
 #include "GenericPlatform/HttpRequestPayload.h"
+#include "GenericPlatform/HttpResponseCommon.h"
 #include "HAL/ThreadSafeBool.h"
 
 class FCurlHttpResponse;
@@ -145,7 +145,6 @@ public:
 
 	//~ Begin IHttpBase Interface
 	virtual FString GetURL() const override;
-	virtual FString GetURLParameter(const FString& ParameterName) const override;
 	virtual FString GetHeader(const FString& HeaderName) const override;
 	virtual TArray<FString> GetAllHeaders() const override;
 	virtual FString GetContentType() const override;
@@ -413,21 +412,13 @@ private:
 /**
  * Curl implementation of an HTTP response
  */
-class FCurlHttpResponse : public IHttpResponse
+class FCurlHttpResponse : public FHttpResponseCommon
 {
-private:
-
-	/** Request that owns this response */
-	FCurlHttpRequest& Request;
-
-
 public:
 	// implementation friends
 	friend class FCurlHttpRequest;
 
 	//~ Begin IHttpBase Interface
-	virtual FString GetURL() const override;
-	virtual FString GetURLParameter(const FString& ParameterName) const override;
 	virtual FString GetHeader(const FString& HeaderName) const override;
 	virtual TArray<FString> GetAllHeaders() const override;	
 	virtual FString GetContentType() const override;
@@ -445,15 +436,9 @@ public:
 	 *
 	 * @param InRequest - original request that created this response
 	 */
-	FCurlHttpResponse(FCurlHttpRequest& InRequest);
-
-	/**
-	 * Destructor
-	 */
-	virtual ~FCurlHttpResponse();
+	FCurlHttpResponse(const FCurlHttpRequest& InRequest);
 
 private:
-
 	/** BYTE array to fill in as the response is read via didReceiveData */
 	TArray<uint8> Payload;
 	/** The stream to receive response body */
diff --git a/Engine/Source/Runtime/Online/HTTP/Private/GenericPlatform/HttpRequestCommon.cpp b/Engine/Source/Runtime/Online/HTTP/Private/GenericPlatform/HttpRequestCommon.cpp
--- a/Engine/Source/Runtime/Online/HTTP/Private/GenericPlatform/HttpRequestCommon.cpp
+++ b/Engine/Source/Runtime/Online/HTTP/Private/GenericPlatform/HttpRequestCommon.cpp
@@ -4,6 +4,16 @@
 #include "Http.h"
 #include "HttpManager.h"
 
+FString FHttpRequestCommon::GetURLParameter(const FString& ParameterName) const
+{
+	FString ReturnValue;
+	if (TOptional<FString> OptionalParameterValue = FGenericPlatformHttp::GetUrlParameter(GetURL(), ParameterName))
+	{
+		ReturnValue = MoveTemp(OptionalParameterValue.GetValue());
+	}
+	return ReturnValue;
+}
+
 EHttpRequestStatus::Type FHttpRequestCommon::GetStatus() const
 {
 	return CompletionStatus;
@@ -25,6 +35,12 @@ bool FHttpRequestCommon::PreCheck() const
 		return false;
 	}
 
+	if (GetVerb().IsEmpty())
+	{
+		UE_LOG(LogHttp, Warning, TEXT("ProcessRequest failed. No Verb was specified."));
+		return false;
+	}
+
 	// Make sure the URL is parsed correctly with a valid HTTP scheme
 	const FString URL = GetURL();
 	if (!URL.StartsWith(TEXT("http://")) && !URL.StartsWith(TEXT("https://")))
diff --git a/Engine/Source/Runtime/Online/HTTP/Private/GenericPlatform/HttpResponseCommon.cpp b/Engine/Source/Runtime/Online/HTTP/Private/GenericPlatform/HttpResponseCommon.cpp
new file mode 100644
--- /dev/null
+++ b/Engine/Source/Runtime/Online/HTTP/Private/GenericPlatform/HttpResponseCommon.cpp
@@ -0,0 +1,24 @@
+// Copyright Epic Games, Inc. All Rights Reserved.
+
+#include "GenericPlatform/HttpResponseCommon.h"
+#include "GenericPlatform/GenericPlatformHttp.h"
+
+FHttpResponseCommon::FHttpResponseCommon(const FString& InURL)
+	: URL(InURL)
+{
+}
+
+FString FHttpResponseCommon::GetURLParameter(const FString& ParameterName) const
+{
+	FString ReturnValue;
+	if (TOptional<FString> OptionalParameterValue = FGenericPlatformHttp::GetUrlParameter(URL, ParameterName))
+	{
+		ReturnValue = MoveTemp(OptionalParameterValue.GetValue());
+	}
+	return ReturnValue;
+}
+
+FString FHttpResponseCommon::GetURL() const
+{
+	return URL;
+}
diff --git a/Engine/Source/Runtime/Online/HTTP/Public/GenericPlatform/HttpRequestCommon.h b/Engine/Source/Runtime/Online/HTTP/Public/GenericPlatform/HttpRequestCommon.h
--- a/Engine/Source/Runtime/Online/HTTP/Public/GenericPlatform/HttpRequestCommon.h
+++ b/Engine/Source/Runtime/Online/HTTP/Public/GenericPlatform/HttpRequestCommon.h
@@ -10,6 +10,9 @@
 class FHttpRequestCommon : public FHttpRequestImpl
 {
 public:
+	// IHttpBase
+	HTTP_API virtual FString GetURLParameter(const FString& ParameterName) const override;
+
 	// IHttpRequest
 	HTTP_API virtual EHttpRequestStatus::Type GetStatus() const override;
 	HTTP_API virtual void SetDelegateThreadPolicy(EHttpRequestDelegateThreadPolicy InDelegateThreadPolicy) override;
diff --git a/Engine/Source/Runtime/Online/HTTP/Public/GenericPlatform/HttpResponseCommon.h b/Engine/Source/Runtime/Online/HTTP/Public/GenericPlatform/HttpResponseCommon.h
new file mode 100644
--- /dev/null
+++ b/Engine/Source/Runtime/Online/HTTP/Public/GenericPlatform/HttpResponseCommon.h
@@ -0,0 +1,21 @@
+// Copyright Epic Games, Inc. All Rights Reserved.
+
+#pragma once
+
+#include "Interfaces/IHttpResponse.h"
+
+/**
+ * Contains implementation of some common functions that don't vary between implementations of different platforms
+ */
+class FHttpResponseCommon : public IHttpResponse
+{
+public:
+	HTTP_API FHttpResponseCommon(const FString& InURL);
+
+	// IHttpBase
+	HTTP_API virtual FString GetURLParameter(const FString& ParameterName) const override;
+	HTTP_API virtual FString GetURL() const override;
+
+protected:
+	FString URL;
+};
