Remove dependency from WebTests.
 - Adds "Json" dependency instead

diff --git a/Engine/Source/Programs/WebTests/WebTests.Build.cs b/Engine/Source/Programs/WebTests/WebTests.Build.cs
--- a/Engine/Source/Programs/WebTests/WebTests.Build.cs
+++ b/Engine/Source/Programs/WebTests/WebTests.Build.cs
@@ -12,8 +12,7 @@ public class WebTests : TestModuleRules
 				"Core",
 				"HTTP",
 				"WebSockets",
-				"Chaos",
-				"ChaosUserDataPT"
+				"Json"
 			});
 		
 		UpdateBuildGraphPropertiesFile(new Metadata() { 
