Make retry system don't rely on Tick from game thread;
Make http tests can run with retry system enabled/disabled.

#tests Passed WebTests, also tried on Windows to reach main menu

diff --git a/Engine/Source/Programs/WebTests/Private/TestHttp.cpp b/Engine/Source/Programs/WebTests/Private/TestHttp.cpp
--- a/Engine/Source/Programs/WebTests/Private/TestHttp.cpp
+++ b/Engine/Source/Programs/WebTests/Private/TestHttp.cpp
@@ -38,6 +38,7 @@ public:
 		: WebServerIp(TEXT("127.0.0.1"))
 		, WebServerHttpPort(8000)
 		, bRunHeavyTests(false)
+		, bRetryEnabled(true)
 		, OldVerbosity(LogHttp.GetVerbosity())
 	{
 		ParseSettingsFromCommandLine();
@@ -65,6 +66,7 @@ public:
 	{
 		FParse::Value(FCommandLine::Get(), TEXT("web_server_ip"), WebServerIp);
 		FParse::Bool(FCommandLine::Get(), TEXT("run_heavy_tests"), bRunHeavyTests);
+		FParse::Bool(FCommandLine::Get(), TEXT("retry_enabled"), bRetryEnabled);
 	}
 
 	void DisableWarningsInThisTest()
@@ -83,6 +85,7 @@ public:
 	FHttpModule* HttpModule;
 
 	bool bRunHeavyTests;
+	bool bRetryEnabled;
 	ELogVerbosity::Type OldVerbosity;
 };
 
@@ -113,57 +116,15 @@ TEST_CASE_METHOD(FHttpModuleTestFixture, "Shutdown http module without issue whe
 	HttpModule->GetHttpManager().Tick(0.0f);
 }
 
-class FWaitUntilQuitFromTestFixture : public FHttpModuleTestFixture
+class FMockRetryManager : public FHttpRetrySystem::FManager
 {
 public:
-	FWaitUntilQuitFromTestFixture()
-	{
-	}
-
-	~FWaitUntilQuitFromTestFixture()
-	{
-		WaitUntilQuitFromTest();
-	}
-
-	void WaitUntilQuitFromTest()
-	{
-		while (!bQuitRequested)
-		{
-			HttpModule->GetHttpManager().Tick(TickFrequency);
-			FPlatformProcess::Sleep(TickFrequency);
-		}
-	}
-
-	float TickFrequency = 1.0f / 60; /*60 FPS*/;
-	bool bQuitRequested = false;
+	using FHttpRetrySystem::FManager::FManager;
+	using FHttpRetrySystem::FManager::RequestList;
+	using FHttpRetrySystem::FManager::FHttpRetryRequestEntry;
+	using FHttpRetrySystem::FManager::RetryTimeoutRelativeSecondsDefault;
 };
 
-TEST_CASE_METHOD(FWaitUntilQuitFromTestFixture, "Http request can be reused", HTTP_TAG)
-{
-	TSharedRef<IHttpRequest> HttpRequest = HttpModule->CreateRequest();
-	HttpRequest->SetURL(UrlToTestMethods());
-	HttpRequest->SetVerb(TEXT("POST"));
-
-	HttpRequest->OnProcessRequestComplete().BindLambda([this](FHttpRequestPtr HttpRequest, FHttpResponsePtr HttpResponse, bool bSucceeded) {
-		CHECK(bSucceeded);
-		CHECK(HttpResponse->GetResponseCode() == 200);
-
-		uint32 Chunks = 3;
-		uint32 ChunkSize = 1024;
-		HttpRequest->SetURL(UrlStreamDownload(Chunks, ChunkSize));
-		HttpRequest->SetVerb(TEXT("GET"));
-		HttpRequest->OnProcessRequestComplete().BindLambda([this, Chunks, ChunkSize](FHttpRequestPtr HttpRequest, FHttpResponsePtr HttpResponse, bool bSucceeded) {
-			CHECK(bSucceeded);
-			REQUIRE(HttpResponse != nullptr);
-			CHECK(HttpResponse->GetResponseCode() == 200);
-			CHECK(HttpResponse->GetContentLength() == Chunks * ChunkSize);
-			bQuitRequested = true;
-		});
-		HttpRequest->ProcessRequest();
-	});
-	HttpRequest->ProcessRequest();
-}
-
 class FWaitUntilCompleteHttpFixture : public FHttpModuleTestFixture
 {
 public:
@@ -171,6 +132,12 @@ public:
 	{
 		HttpModule->GetHttpManager().SetRequestAddedDelegate(FHttpManagerRequestAddedDelegate::CreateRaw(this, &FWaitUntilCompleteHttpFixture::OnRequestAdded));
 		HttpModule->GetHttpManager().SetRequestCompletedDelegate(FHttpManagerRequestCompletedDelegate::CreateRaw(this, &FWaitUntilCompleteHttpFixture::OnRequestCompleted));
+
+		if (bRetryEnabled)
+		{
+			HttpRetryManager = MakeShared<FMockRetryManager>(FHttpRetrySystem::FRetryLimitCountSetting(RetryLimitCount), FHttpRetrySystem::FRetryTimeoutRelativeSecondsSetting(/*RetryTimeoutRelativeSeconds*/));
+			HttpRetryManager->RetryTimeoutRelativeSecondsDefault = 2.0f; // Value is set so that retry system will use it at high level, will no longer wait for ever for HTTP code
+		}
 	}
 
 	~FWaitUntilCompleteHttpFixture()
@@ -204,13 +171,21 @@ public:
 		HttpModule->GetHttpManager().Tick(TickFrequency);
 	}
 
+	TSharedRef<IHttpRequest> CreateRequest()
+	{
+		return bRetryEnabled ? HttpRetryManager->CreateRequest() : HttpModule->CreateRequest();
+	}
+
 	std::atomic<int32> OngoingRequests = 0;
 	float TickFrequency = 1.0f / 60; /*60 FPS*/;
+
+	uint32 RetryLimitCount = 2;
+	TSharedPtr<FMockRetryManager> HttpRetryManager;
 };
 
 TEST_CASE_METHOD(FWaitUntilCompleteHttpFixture, "Http Methods", HTTP_TAG)
 {
-	TSharedRef<IHttpRequest> HttpRequest = HttpModule->CreateRequest();
+	TSharedRef<IHttpRequest> HttpRequest = CreateRequest();
 	CHECK(HttpRequest->GetVerb() == TEXT("GET"));
 
 	HttpRequest->SetURL(UrlToTestMethods());
@@ -245,7 +220,7 @@ TEST_CASE_METHOD(FWaitUntilCompleteHttpFixture, "Http Methods", HTTP_TAG)
 
 TEST_CASE_METHOD(FWaitUntilCompleteHttpFixture, "Get large response content without chunks", HTTP_TAG)
 {
-	TSharedRef<IHttpRequest> HttpRequest = HttpModule->CreateRequest();
+	TSharedRef<IHttpRequest> HttpRequest = CreateRequest();
 	HttpRequest->SetURL(FString::Format(TEXT("{0}/get_large_response_without_chunks/{1}/"), { *UrlHttpTests(), 1024 * 1024/*bytes_number*/}));
 	HttpRequest->SetVerb(TEXT("GET"));
 	HttpRequest->OnProcessRequestComplete().BindLambda([](FHttpRequestPtr HttpRequest, FHttpResponsePtr HttpResponse, bool bSucceeded) {
@@ -259,7 +234,7 @@ TEST_CASE_METHOD(FWaitUntilCompleteHttpFixture, "Get large response content with
 // TODO: Enable this after finding a more reliable way of simulate timeout
 //TEST_CASE_METHOD(FWaitUntilCompleteHttpFixture, "Http request connect timeout", HTTP_TAG)
 //{
-//	TSharedRef<IHttpRequest> HttpRequest = HttpModule->CreateRequest();
+//	TSharedRef<IHttpRequest> HttpRequest = CreateRequest();
 //	HttpRequest->SetURL(UrlWithInvalidPortToTestConnectTimeout());
 //	HttpRequest->SetVerb(TEXT("GET"));
 //	HttpRequest->SetTimeout(7);
@@ -282,7 +257,7 @@ TEST_CASE_METHOD(FWaitUntilCompleteHttpFixture, "Streaming http download", HTTP_
 	uint32 Chunks = 3;
 	uint32 ChunkSize = 1024*1024;
 
-	TSharedRef<IHttpRequest> HttpRequest = HttpModule->CreateRequest();
+	TSharedRef<IHttpRequest> HttpRequest = CreateRequest();
 	HttpRequest->SetURL(UrlStreamDownload(Chunks, ChunkSize));
 	HttpRequest->SetVerb(TEXT("GET"));
 
@@ -432,7 +407,7 @@ TEST_CASE_METHOD(FWaitUntilCompleteHttpFixture, "Can run parallel stream downloa
 
 	for (int i = 0; i < 3; ++i)
 	{
-		TSharedRef<IHttpRequest> HttpRequest = HttpModule->CreateRequest();
+		TSharedRef<IHttpRequest> HttpRequest = CreateRequest();
 		HttpRequest->SetURL(UrlStreamDownload(Chunks, ChunkSize));
 		HttpRequest->SetVerb(TEXT("GET"));
 		HttpRequest->OnProcessRequestComplete().BindLambda([Chunks, ChunkSize](FHttpRequestPtr HttpRequest, FHttpResponsePtr HttpResponse, bool bSucceeded) {
@@ -455,7 +430,7 @@ TEST_CASE_METHOD(FWaitUntilCompleteHttpFixture, "Can download big file exceeds 3
 	uint64 Chunks = 5 * 1024;
 	uint64 ChunkSize = 1024 * 1024;
 
-	TSharedRef<IHttpRequest> HttpRequest = HttpModule->CreateRequest();
+	TSharedRef<IHttpRequest> HttpRequest = CreateRequest();
 	HttpRequest->SetURL(UrlStreamDownload(Chunks, ChunkSize));
 	HttpRequest->SetVerb(TEXT("GET"));
 
@@ -480,7 +455,7 @@ TEST_CASE_METHOD(FWaitUntilCompleteHttpFixture, "Can download big file exceeds 3
 
 TEST_CASE_METHOD(FWaitUntilCompleteHttpFixture, "Streaming http upload from memory", HTTP_TAG)
 {
-	TSharedRef<IHttpRequest> HttpRequest = HttpModule->CreateRequest();
+	TSharedRef<IHttpRequest> HttpRequest = CreateRequest();
 	HttpRequest->SetURL(FString::Format(TEXT("{0}/streaming_upload_post"), { *UrlHttpTests() }));
 	HttpRequest->SetVerb(TEXT("POST"));
 
@@ -574,7 +549,7 @@ TEST_CASE_METHOD(FWaitUntilCompleteHttpFixture, "Can upload big file exceeds 32
 	const uint64 TotalSize = 2147483647;
 	TSharedRef<FTestHttpUploadStream> Stream = MakeShared<FTestHttpUploadStream>(TotalSize);
 
-	TSharedRef<IHttpRequest> HttpRequest = HttpModule->CreateRequest();
+	TSharedRef<IHttpRequest> HttpRequest = CreateRequest();
 	HttpRequest->SetURL(FString::Format(TEXT("{0}/streaming_upload_put"), { *UrlHttpTests() }));
 	HttpRequest->SetVerb(TEXT("PUT"));
 	HttpRequest->SetContentFromStream(Stream);
@@ -603,7 +578,7 @@ TEST_CASE_METHOD(FWaitUntilCompleteHttpFixture, "Streaming http upload from file
 	FileToWrite->Close();
 	FMemory::Free(FileData);
 
-	TSharedRef<IHttpRequest> HttpRequest = HttpModule->CreateRequest();
+	TSharedRef<IHttpRequest> HttpRequest = CreateRequest();
 	HttpRequest->SetURL(FString::Format(TEXT("{0}/streaming_upload_put"), { *UrlHttpTests() }));
 	HttpRequest->SetVerb(TEXT("PUT"));
 	HttpRequest->SetHeader(TEXT("Content-Disposition"), TEXT("attachment;filename=TestStreamUpload.dat"));
@@ -618,7 +593,7 @@ TEST_CASE_METHOD(FWaitUntilCompleteHttpFixture, "Streaming http upload from file
 
 TEST_CASE_METHOD(FWaitUntilCompleteHttpFixture, "Redirect enabled by default and can work well", HTTP_TAG)
 {
-	TSharedRef<IHttpRequest> HttpRequest = HttpModule->CreateRequest();
+	TSharedRef<IHttpRequest> HttpRequest = CreateRequest();
 	HttpRequest->SetURL(FString::Format(TEXT("{0}/redirect_from"), { *UrlHttpTests() }));
 	HttpRequest->SetVerb(TEXT("GET"));
 	HttpRequest->OnProcessRequestComplete().BindLambda([](FHttpRequestPtr HttpRequest, FHttpResponsePtr HttpResponse, bool bSucceeded) {
@@ -628,6 +603,58 @@ TEST_CASE_METHOD(FWaitUntilCompleteHttpFixture, "Redirect enabled by default and
 	HttpRequest->ProcessRequest();
 }
 
+
+class FWaitUntilQuitFromTestFixture : public FWaitUntilCompleteHttpFixture
+{
+public:
+	FWaitUntilQuitFromTestFixture()
+	{
+	}
+
+	~FWaitUntilQuitFromTestFixture()
+	{
+		WaitUntilQuitFromTest();
+	}
+
+	void WaitUntilQuitFromTest()
+	{
+		while (!bQuitRequested)
+		{
+			HttpModule->GetHttpManager().Tick(TickFrequency);
+			FPlatformProcess::Sleep(TickFrequency);
+		}
+	}
+
+	bool bQuitRequested = false;
+};
+
+TEST_CASE_METHOD(FWaitUntilQuitFromTestFixture, "Http request can be reused", HTTP_TAG)
+{
+	TSharedRef<IHttpRequest> HttpRequest = CreateRequest();
+	HttpRequest->SetURL(UrlToTestMethods());
+	HttpRequest->SetVerb(TEXT("POST"));
+
+	HttpRequest->OnProcessRequestComplete().BindLambda([this](FHttpRequestPtr HttpRequest, FHttpResponsePtr HttpResponse, bool bSucceeded) {
+		CHECK(bSucceeded);
+		CHECK(HttpResponse->GetResponseCode() == 200);
+
+		uint32 Chunks = 3;
+		uint32 ChunkSize = 1024;
+		HttpRequest->SetURL(UrlStreamDownload(Chunks, ChunkSize));
+		HttpRequest->SetVerb(TEXT("GET"));
+		HttpRequest->OnProcessRequestComplete().BindLambda([this, Chunks, ChunkSize](FHttpRequestPtr HttpRequest, FHttpResponsePtr HttpResponse, bool bSucceeded) {
+			CHECK(bSucceeded);
+			REQUIRE(HttpResponse != nullptr);
+			CHECK(HttpResponse->GetResponseCode() == 200);
+			CHECK(HttpResponse->GetContentLength() == Chunks * ChunkSize);
+			bQuitRequested = true;
+		});
+		HttpRequest->ProcessRequest();
+	});
+	HttpRequest->ProcessRequest();
+}
+
+
 // Response shared ptr should be able to be kept by user code and valid to access without http request
 class FValidateResponseDependencyFixture : public FWaitUntilCompleteHttpFixture
 {
@@ -748,7 +775,7 @@ public:
 TEST_CASE_METHOD(FWaitThreadedHttpFixture, "Http streaming download request can work in non game thread", HTTP_TAG)
 {
 	ThreadedHttpRunnable.OnRunFromThread().BindLambda([this]() {
-		TSharedRef<IHttpRequest> HttpRequest = HttpModule->CreateRequest();
+		TSharedRef<IHttpRequest> HttpRequest = CreateRequest();
 		HttpRequest->SetURL(UrlStreamDownload(3/*Chunks*/, 1024/*ChunkSize*/));
 		HttpRequest->SetVerb(TEXT("GET"));
 		HttpRequest->SetDelegateThreadPolicy(EHttpRequestDelegateThreadPolicy::CompleteOnHttpThread);
@@ -784,7 +811,7 @@ TEST_CASE_METHOD(FWaitThreadedHttpFixture, "Http download request progress callb
 {
 	std::atomic<bool> bRequestProgressTriggered = false;
 	ThreadedHttpRunnable.OnRunFromThread().BindLambda([this, &bRequestProgressTriggered]() {
-		TSharedRef<IHttpRequest> HttpRequest = HttpModule->CreateRequest();
+		TSharedRef<IHttpRequest> HttpRequest = CreateRequest();
 		HttpRequest->SetURL(UrlStreamDownload(10/*Chunks*/, 1024*1024/*ChunkSize*/));
 		HttpRequest->SetVerb(TEXT("GET"));
 
@@ -858,37 +885,52 @@ void SetupURLRequestFilter(FHttpModule* HttpModule)
 }
 }
 
-TEST_CASE_METHOD(FWaitUntilCompleteHttpFixture, "Http request pre check will fail by thread policy", HTTP_TAG)
+// Pre-check failed requests won't be added into http manager, so it can't rely on the requested added/completed callback in FWaitUntilCompleteHttpFixture
+TEST_CASE_METHOD(FWaitUntilQuitFromTestFixture, "Http request pre check will fail by thread policy", HTTP_TAG)
 {
 	DisableWarningsInThisTest();
 
 	// Pre check will fail when domain is not allowed
 	UE::TestHttp::SetupURLRequestFilter(HttpModule);
 
-	TSharedRef<IHttpRequest> HttpRequest = HttpModule->CreateRequest();
+	TSharedRef<IHttpRequest> HttpRequest = CreateRequest();
 	HttpRequest->SetVerb(TEXT("GET"));
 	HttpRequest->SetURL(UrlToTestMethods());
 
 	SECTION("on game thread")
 	{
-		HttpRequest->OnProcessRequestComplete().BindLambda([](FHttpRequestPtr HttpRequest, FHttpResponsePtr HttpResponse, bool bSucceeded) {
+		HttpRequest->OnProcessRequestComplete().BindLambda([this](FHttpRequestPtr HttpRequest, FHttpResponsePtr HttpResponse, bool bSucceeded) {
 			CHECK(IsInGameThread());
 			CHECK(!bSucceeded);
+			bQuitRequested = true;
 		});
 	}
 	SECTION("on http thread")
 	{
 		HttpRequest->SetDelegateThreadPolicy(EHttpRequestDelegateThreadPolicy::CompleteOnHttpThread);
-		HttpRequest->OnProcessRequestComplete().BindLambda([](FHttpRequestPtr HttpRequest, FHttpResponsePtr HttpResponse, bool bSucceeded) {
+		HttpRequest->OnProcessRequestComplete().BindLambda([this](FHttpRequestPtr HttpRequest, FHttpResponsePtr HttpResponse, bool bSucceeded) {
 			CHECK(!IsInGameThread());
 			CHECK(!bSucceeded);
+			bQuitRequested = true;
 		});
 	}
 
 	HttpRequest->ProcessRequest();
 }
 
-TEST_CASE_METHOD(FWaitThreadedHttpFixture, "Threaded http request pre check will fail by thread policy", HTTP_TAG)
+class FWaitUnitilQuitFromTestThreadedFixture : public FWaitUntilQuitFromTestFixture
+{
+public:
+	~FWaitUnitilQuitFromTestThreadedFixture()
+	{
+		WaitUntilQuitFromTest();
+	}
+
+	FThreadedHttpRunnable ThreadedHttpRunnable;
+};
+
+// Pre-check failed requests won't be added into http manager, so it can't rely on the requested added/completed callback in FWaitUntilCompleteHttpFixture
+TEST_CASE_METHOD(FWaitUnitilQuitFromTestThreadedFixture, "Threaded http request pre check will fail by thread policy", HTTP_TAG)
 {
 	DisableWarningsInThisTest();
 
@@ -896,23 +938,25 @@ TEST_CASE_METHOD(FWaitThreadedHttpFixture, "Threaded http request pre check will
 		// Pre check will fail when domain is not allowed
 		UE::TestHttp::SetupURLRequestFilter(HttpModule);
 
-		TSharedRef<IHttpRequest> HttpRequest = HttpModule->CreateRequest();
+		TSharedRef<IHttpRequest> HttpRequest = CreateRequest();
 		HttpRequest->SetVerb(TEXT("GET"));
 		HttpRequest->SetURL(UrlToTestMethods());
 
 		SECTION("on game thread")
 		{
-			HttpRequest->OnProcessRequestComplete().BindLambda([](FHttpRequestPtr HttpRequest, FHttpResponsePtr HttpResponse, bool bSucceeded) {
+			HttpRequest->OnProcessRequestComplete().BindLambda([this](FHttpRequestPtr HttpRequest, FHttpResponsePtr HttpResponse, bool bSucceeded) {
 				CHECK(IsInGameThread());
 				CHECK(!bSucceeded);
+				bQuitRequested = true;
 			});
 		}
 		SECTION("on http thread")
 		{
 			HttpRequest->SetDelegateThreadPolicy(EHttpRequestDelegateThreadPolicy::CompleteOnHttpThread);
-			HttpRequest->OnProcessRequestComplete().BindLambda([](FHttpRequestPtr HttpRequest, FHttpResponsePtr HttpResponse, bool bSucceeded) {
+			HttpRequest->OnProcessRequestComplete().BindLambda([this](FHttpRequestPtr HttpRequest, FHttpResponsePtr HttpResponse, bool bSucceeded) {
 				CHECK(!IsInGameThread());
 				CHECK(!bSucceeded);
+				bQuitRequested = true;
 			});
 		}
 
@@ -926,7 +970,7 @@ TEST_CASE_METHOD(FWaitUntilCompleteHttpFixture, "Cancel http request connect bef
 {
 	DisableWarningsInThisTest();
 
-	TSharedRef<IHttpRequest, ESPMode::ThreadSafe> HttpRequest = HttpModule->CreateRequest();
+	TSharedRef<IHttpRequest, ESPMode::ThreadSafe> HttpRequest = CreateRequest();
 	HttpRequest->SetURL(UrlWithInvalidPortToTestConnectTimeout());
 	HttpRequest->SetVerb(TEXT("GET"));
 	HttpRequest->SetTimeout(7);
@@ -943,107 +987,48 @@ TEST_CASE_METHOD(FWaitUntilCompleteHttpFixture, "Cancel http request connect bef
 	HttpRequest->CancelRequest();
 }
 
-class FMockRetryManager : public FHttpRetrySystem::FManager
-{
-public:
-	using FHttpRetrySystem::FManager::FManager;
-	using FHttpRetrySystem::FManager::RequestList;
-	using FHttpRetrySystem::FManager::FHttpRetryRequestEntry;
-	using FHttpRetrySystem::FManager::RetryTimeoutRelativeSecondsDefault;
-};
-
-class FRetryHttpFixture : public FWaitUntilCompleteHttpFixture
+class FThreadedBatchRequestsFixture : public FWaitThreadedHttpFixture
 {
 public:
-	FRetryHttpFixture()
+	void LaunchBatchRequests(uint32 BatchSize)
 	{
-		HttpRetryManager = MakeShared<FMockRetryManager>(FHttpRetrySystem::FRetryLimitCountSetting(RetryLimitCount), FHttpRetrySystem::FRetryTimeoutRelativeSecondsSetting(/*RetryTimeoutRelativeSeconds*/));
-		HttpRetryManager->RetryTimeoutRelativeSecondsDefault = 2.0f; // Value is set so that retry system will use it at high level, will no longer wait for ever for HTTP code
-	}
-
-	~FRetryHttpFixture()
-	{
-		// Make sure not destroy retry manager before waiting in FWaitUntilCompleteHttpFixture
-		while (HttpRetryManager->RequestList.Num())
+		for (uint32 i = 0; i < BatchSize; ++i)
 		{
-			HttpRetryManager->Update();
-			HttpModule->GetHttpManager().Tick(TickFrequency);
-			FPlatformProcess::Sleep(TickFrequency);
+			TSharedRef<IHttpRequest> HttpRequest = CreateRequest();
+			HttpRequest->SetURL(UrlStreamDownload(3, 1024*1024));
+			HttpRequest->SetVerb(TEXT("GET"));
+			HttpRequest->ProcessRequest();
 		}
 	}
 
-	TSharedPtr<FMockRetryManager> HttpRetryManager;
-	uint32 RetryLimitCount = 2;
-
-	FThreadedHttpRunnable ThreadedHttpRunnable;
-};
-
-// TODO: Currently broken, will fix this in retry system
-//TEST_CASE_METHOD(FRetryHttpFixture, "Http request adaptor in retry system can work in non game thread", HTTP_TAG)
-//{
-//	ThreadedHttpRunnable.OnRunFromThread().BindLambda([this]() {
-//		TSharedRef<IHttpRequest> HttpRequest = HttpRetryManager->CreateRequest();
-//		HttpRequest->SetURL(UrlStreamDownload(3/*Chunks*/, 1024/*ChunkSize*/));
-//		HttpRequest->SetVerb(TEXT("GET"));
-//		HttpRequest->SetDelegateThreadPolicy(EHttpRequestDelegateThreadPolicy::CompleteOnHttpThread);
-//
-//		class FTestHttpReceiveStream final : public FArchive
-//		{
-//		public:
-//			virtual void Serialize(void* V, int64 Length) override
-//			{
-//				// No matter what's the thread policy, Serialize always get called in http thread.
-//				CHECK(!IsInGameThread());
-//			}
-//		};
-//		CHECK(HttpRequest->SetResponseBodyReceiveStream(MakeShared<FTestHttpReceiveStream>()));
-//
-//		HttpRequest->OnProcessRequestComplete().BindLambda([this](FHttpRequestPtr HttpRequest, FHttpResponsePtr HttpResponse, bool bSucceeded) {
-//			// EHttpRequestDelegateThreadPolicy::CompleteOnHttpThread was used, so not in game thread here
-//			CHECK(!IsInGameThread());
-//			CHECK(bSucceeded);
-//			CHECK(HttpResponse->GetResponseCode() == 200);
-//			ThreadedHttpRunnable.UnblockGameThread();
-//		});
-//
-//		HttpRequest->ProcessRequest();
-//	});
-//
-//	ThreadedHttpRunnable.StartTestHttpThread(true/*bBlockGameThread*/);
-//}
-
-class FRetryHttpManagerThreadedRequestsFixture : public FRetryHttpFixture
-{
-public:
-	void LaunchDownloadRequests()
+	void BlockUntilFlushed()
 	{
-		for (int32 i = 0; i < 10; ++i)
+		if (bRetryEnabled)
 		{
-			TSharedRef<IHttpRequest> HttpRequest = HttpRetryManager->CreateRequest();
-			HttpRequest->SetURL(UrlStreamDownload(3, 1024*1024));
-			HttpRequest->SetVerb(TEXT("GET"));
-			HttpRequest->ProcessRequest();
+			HttpRetryManager->BlockUntilFlushed(5.0);
+		}
+		else
+		{
+			HttpModule->GetHttpManager().Flush(EHttpFlushReason::Default);
 		}
 	}
 };
 
-TEST_CASE_METHOD(FRetryHttpManagerThreadedRequestsFixture, "Retry manager is thread safe", HTTP_TAG)
+TEST_CASE_METHOD(FThreadedBatchRequestsFixture, "Retry manager and http manager is thread safe", HTTP_TAG)
 {
 	DisableWarningsInThisTest();
 
 	ThreadedHttpRunnable.OnRunFromThread().BindLambda([this]() {
-		LaunchDownloadRequests();
-		HttpRetryManager->BlockUntilFlushed(5.0);
+		LaunchBatchRequests(10);
+		BlockUntilFlushed();
 	});
 	ThreadedHttpRunnable.StartTestHttpThread(false/*bBlockGameThread*/);
 
-	LaunchDownloadRequests();
-	HttpRetryManager->BlockUntilFlushed(5.0);
+	LaunchBatchRequests(10);
+	BlockUntilFlushed();
 }
 
 
 // TODO: Add cancel test, with multiple cancel calls
 
-// TODO: Add a test case to make sure once Request completed and destroyed, if the user code still keeps the shared ptr of Response, any call to Response shouldn't have dependency to Request
-
 // TODO: Add test case to validate header received callback can be received in http/game thread, and can be received before the request complete
diff --git a/Engine/Source/Runtime/Analytics/AnalyticsET/Private/IAnalyticsProviderET.cpp b/Engine/Source/Runtime/Analytics/AnalyticsET/Private/IAnalyticsProviderET.cpp
--- a/Engine/Source/Runtime/Analytics/AnalyticsET/Private/IAnalyticsProviderET.cpp
+++ b/Engine/Source/Runtime/Analytics/AnalyticsET/Private/IAnalyticsProviderET.cpp
@@ -314,7 +314,9 @@ bool FAnalyticsProviderET::Tick(float DeltaSeconds)
 {
 	QUICK_SCOPE_CYCLE_COUNTER(STAT_FAnalyticsProviderET_Tick);
 
+PRAGMA_DISABLE_DEPRECATION_WARNINGS
 	HttpRetryManager->Update();
+PRAGMA_ENABLE_DEPRECATION_WARNINGS
 
 	// hold a lock the entire time here because we're making several calls to the event cache that we need to be consistent when we decide to flush.
 	// With more care, we can likely avoid holding this lock the entire time.
diff --git a/Engine/Source/Runtime/Online/HTTP/Private/Curl/CurlHttp.cpp b/Engine/Source/Runtime/Online/HTTP/Private/Curl/CurlHttp.cpp
--- a/Engine/Source/Runtime/Online/HTTP/Private/Curl/CurlHttp.cpp
+++ b/Engine/Source/Runtime/Online/HTTP/Private/Curl/CurlHttp.cpp
@@ -978,7 +978,8 @@ bool FCurlHttpRequest::ProcessRequest()
 
 	if (!PreCheck() || !SetupRequest())
 	{
-		return FinishRequestNotInHttpManager();
+		FinishRequestNotInHttpManager();
+		return false;
 	}
 
 	// Clear the info cache log so we don't output messages from previous requests when reusing/retrying a request
@@ -1260,7 +1261,7 @@ void FCurlHttpRequest::FinishRequest()
 		{
 			UE_LOG(LogHttp, Warning, TEXT("%p: request failed, libcurl multi error: %d (%s)"), this, (int32)CurlAddToMultiResult, ANSI_TO_TCHAR(curl_multi_strerror(CurlAddToMultiResult)));
 		}
-		else
+		else if (CurlCompletionResult != CURLE_OK)
 		{
 			UE_LOG(LogHttp, Warning, TEXT("%p: request failed, libcurl error: %d (%s)"), this, (int32)CurlCompletionResult, ANSI_TO_TCHAR(curl_easy_strerror(CurlCompletionResult)));
 		}
diff --git a/Engine/Source/Runtime/Online/HTTP/Private/EventLoopHttpThread.cpp b/Engine/Source/Runtime/Online/HTTP/Private/EventLoopHttpThread.cpp
--- a/Engine/Source/Runtime/Online/HTTP/Private/EventLoopHttpThread.cpp
+++ b/Engine/Source/Runtime/Online/HTTP/Private/EventLoopHttpThread.cpp
@@ -159,4 +159,10 @@ void FEventLoopHttpThread::ResetTickTimer()
 	},
 	FTimespan::FromSeconds(FHttpModule::Get().GetHttpEventLoopThreadTickIntervalInSeconds()),
 	true /* repeat */);
+}
+
+void FEventLoopHttpThread::AddHttpThreadTask(TFunction<void()>&& Task, float InDelay)
+{
+	UE::EventLoop::IEventLoop& EventLoop = GetEventLoopChecked();
+	EventLoop.SetTimer(Task, FTimespan::FromSeconds(InDelay));
 }
\ No newline at end of file
diff --git a/Engine/Source/Runtime/Online/HTTP/Private/EventLoopHttpThread.h b/Engine/Source/Runtime/Online/HTTP/Private/EventLoopHttpThread.h
--- a/Engine/Source/Runtime/Online/HTTP/Private/EventLoopHttpThread.h
+++ b/Engine/Source/Runtime/Online/HTTP/Private/EventLoopHttpThread.h
@@ -40,6 +40,8 @@ protected:
 
 	void ResetTickTimer();
 
+	virtual void AddHttpThreadTask(TFunction<void()>&& Task, float InDelay) override;
+
 protected:
 	UE::EventLoop::FTimerHandle RequestTickTimer;
 };
diff --git a/Engine/Source/Runtime/Online/HTTP/Private/HttpManager.cpp b/Engine/Source/Runtime/Online/HTTP/Private/HttpManager.cpp
--- a/Engine/Source/Runtime/Online/HTTP/Private/HttpManager.cpp
+++ b/Engine/Source/Runtime/Online/HTTP/Private/HttpManager.cpp
@@ -284,10 +284,10 @@ void FHttpManager::AddGameThreadTask(TFunction<void()>&& Task)
 	}
 }
 
-void FHttpManager::AddHttpThreadTask(TFunction<void()>&& Task)
+void FHttpManager::AddHttpThreadTask(TFunction<void()>&& Task, float InDelay)
 {
 	check(Thread);
-	Thread->AddHttpThreadTask(MoveTemp(Task));
+	Thread->AddHttpThreadTask(MoveTemp(Task), InDelay);
 }
 
 FHttpThreadBase* FHttpManager::CreateHttpThread()
@@ -547,3 +547,8 @@ void FHttpManager::BroadcastHttpRequestCompleted(const FHttpRequestRef& Request)
 {
 	RequestCompletedDelegate.ExecuteIfBound(Request);
 }
+
+FHttpThreadBase* FHttpManager::GetThread()
+{
+	return Thread;
+}
diff --git a/Engine/Source/Runtime/Online/HTTP/Private/HttpRetrySystem.cpp b/Engine/Source/Runtime/Online/HTTP/Private/HttpRetrySystem.cpp
--- a/Engine/Source/Runtime/Online/HTTP/Private/HttpRetrySystem.cpp
+++ b/Engine/Source/Runtime/Online/HTTP/Private/HttpRetrySystem.cpp
@@ -1,6 +1,7 @@
 // Copyright Epic Games, Inc. All Rights Reserved.
 
 #include "HttpRetrySystem.h"
+#include "HAL/ConsoleManager.h"
 #include "HAL/PlatformTime.h"
 #include "HAL/PlatformProcess.h"
 #include "HAL/LowLevelMemTracker.h"
@@ -8,10 +9,17 @@
 #include "HttpModule.h"
 #include "Http.h"
 #include "HttpManager.h"
+#include "HttpThread.h"
 #include "Stats/Stats.h"
 #include "Misc/ConfigCacheIni.h"
 #include "Misc/CoreDelegates.h"
 
+TAutoConsoleVariable<bool> CVarHttpRetrySystemNonGameThreadSupportEnabled(
+	TEXT("Http.RetrySystemNonGameThreadSupportEnabled"),
+	true,
+	TEXT("Enable retry system non-game thread support")
+);
+
 LLM_DEFINE_TAG(HTTP);
 
 namespace FHttpRetrySystem
@@ -71,7 +79,7 @@ FHttpRetrySystem::FRequest::FRequest(
 	const FHttpRetrySystem::FRetryDomainsPtr& InRetryDomains
 	)
     : FHttpRequestAdapterBase(HttpRequest)
-    , Status(FHttpRetrySystem::FRequest::EStatus::NotStarted)
+    , RetryStatus(FHttpRetrySystem::FRequest::EStatus::NotStarted)
     , RetryLimitCountOverride(InRetryLimitCountOverride)
     , RetryTimeoutRelativeSecondsOverride(InRetryTimeoutRelativeSecondsOverride)
 	, RetryResponseCodes(InRetryResponseCodes)
@@ -109,6 +117,8 @@ bool FHttpRetrySystem::FRequest::ProcessRequest()
 	}
 
 	HttpRequest->OnRequestProgress64().BindThreadSafeSP(RetryRequest, &FHttpRetrySystem::FRequest::HttpOnRequestProgress);
+	HttpRequest->OnProcessRequestComplete().BindThreadSafeSP(RetryRequest, &FHttpRetrySystem::FRequest::HttpOnProcessRequestComplete);
+	HttpRequest->OnHeaderReceived().BindThreadSafeSP(RetryRequest, &FHttpRetrySystem::FRequest::HttpOnHeaderReceived);
 
 	return RetryManager.ProcessRequest(RetryRequest);
 }
@@ -150,11 +160,83 @@ PRAGMA_ENABLE_DEPRECATION_WARNINGS
 	OnRequestProgress64().ExecuteIfBound(AsShared(), BytesSent, BytesRcv);
 }
 
+void FHttpRetrySystem::FRequest::HttpOnProcessRequestComplete(FHttpRequestPtr InHttpRequest, FHttpResponsePtr HttpResponse, bool bSucceeded)
+{
+	if (!CVarHttpRetrySystemNonGameThreadSupportEnabled.GetValueOnAnyThread())
+	{
+		return;
+	}
+
+	TSharedRef<FRequest> SelfPtr = StaticCastSharedRef<FRequest>(AsShared()); // In case no ref after removing from RetryManager
+
+	FScopeLock ScopeLock(&RetryManager.RequestListLock);
+
+	uint32 EntryIndex = RetryManager.RequestList.IndexOfByPredicate([this](const FManager::FHttpRetryRequestEntry& Entry) { return Entry.Request == AsShared(); });
+	check(EntryIndex != INDEX_NONE);
+	FManager::FHttpRetryRequestEntry* HttpRetryRequestEntry = &RetryManager.RequestList[EntryIndex];
+
+	if (RetryStatus == FHttpRetrySystem::FRequest::EStatus::Cancelled)
+	{
+		// Do nothing here
+	}
+	if (GetStatus() == EHttpRequestStatus::Failed_ConnectionError || GetStatus() == EHttpRequestStatus::Failed)
+	{
+		if (GetStatus() == EHttpRequestStatus::Failed_ConnectionError && RetryDomains.IsValid())
+		{
+			MoveToNextRetryDomain();
+		}
+
+		if (RetryManager.ShouldRetry(*HttpRetryRequestEntry) && RetryManager.CanRetry(*HttpRetryRequestEntry))
+		{
+			const double NowAbsoluteSeconds = FPlatformTime::Seconds();
+			float LockoutPeriod = RetryManager.GetLockoutPeriodSeconds(*HttpRetryRequestEntry);
+
+			RetryStatus = FHttpRetrySystem::FRequest::EStatus::ProcessingLockout;
+
+			{
+				QUICK_SCOPE_CYCLE_COUNTER(STAT_FHttpRetrySystem_FManager_Update_OnRequestWillRetry);
+				OnRequestWillRetry().ExecuteIfBound(HttpRetryRequestEntry->Request, GetResponse(), LockoutPeriod);
+			}
+
+			RetryManager.RetryHttpRequestWithDelay(HttpRetryRequestEntry->Request, LockoutPeriod);
+			return;
+		}
+
+		RetryStatus = FHttpRetrySystem::FRequest::EStatus::FailedRetry;
+	}
+	else
+	{
+		RetryStatus = FHttpRetrySystem::FRequest::EStatus::Succeeded;
+	}
+
+	if (HttpRetryRequestEntry->CurrentRetryCount > 0)
+	{
+		FHttpRetrySystem::FManager::FHttpLogVerbosityTracker::Get().DecrementRetriedRequests();
+	}
+
+	RetryManager.RequestList.RemoveAtSwap(EntryIndex);
+
+	OnProcessRequestComplete().ExecuteIfBound(SelfPtr, HttpResponse, bSucceeded);
+}
+
+void FHttpRetrySystem::FRequest::HttpOnHeaderReceived(FHttpRequestPtr Request, const FString& HeaderName, const FString& NewHeaderValue)
+{
+	if (!CVarHttpRetrySystemNonGameThreadSupportEnabled.GetValueOnAnyThread())
+	{
+		return;
+	}
+
+	TSharedRef<FRequest> SelfPtr = StaticCastSharedRef<FRequest>(AsShared());
+	OnHeaderReceived().ExecuteIfBound(SelfPtr, HeaderName, NewHeaderValue);
+}
+
 FHttpRetrySystem::FManager::FManager(const FRetryLimitCountSetting& InRetryLimitCountDefault, const FRetryTimeoutRelativeSecondsSetting& InRetryTimeoutRelativeSecondsDefault)
     : RandomFailureRate(FRandomFailureRateSetting())
     , RetryLimitCountDefault(InRetryLimitCountDefault)
 	, RetryTimeoutRelativeSecondsDefault(InRetryTimeoutRelativeSecondsDefault)
-{}
+{
+	check(FHttpModule::Get().GetHttpManager().GetThread());
+}
 
 FHttpRetrySystem::FManager::~FManager()
 {
@@ -292,18 +374,28 @@ bool FHttpRetrySystem::FManager::HasTimedOut(const FHttpRetryRequestEntry& HttpR
 void FHttpRetrySystem::FManager::RetryHttpRequest(FHttpRetryRequestEntry& RequestEntry)
 {
 	// if this fails the HttpRequest's state will be failed which will cause the retry logic to kick(as expected)
-	const bool bProcessRequestSuccess = RequestEntry.Request->HttpRequest->ProcessRequest();
-	if (bProcessRequestSuccess)
+	if (RequestEntry.CurrentRetryCount == 0)
 	{
-		UE_LOG(LogHttp, Warning, TEXT("Retry %d on %s"), RequestEntry.CurrentRetryCount + 1, *(RequestEntry.Request->GetURL()));
+		FHttpLogVerbosityTracker::Get().IncrementRetriedRequests();
+	}
+	++RequestEntry.CurrentRetryCount;
+	RequestEntry.Request->RetryStatus = FRequest::EStatus::Processing;
+	UE_LOG(LogHttp, Warning, TEXT("Retry %d on %s"), RequestEntry.CurrentRetryCount, *(RequestEntry.Request->GetURL()));
+	RequestEntry.Request->HttpRequest->ProcessRequest();
+}
 
-		if (RequestEntry.CurrentRetryCount == 0)
+void FHttpRetrySystem::FManager::RetryHttpRequestWithDelay(const TSharedRef<FRequest>& Request, float InDelay)
+{
+	FHttpModule::Get().GetHttpManager().AddHttpThreadTask([this, Request]() {
+		FScopeLock ScopeLock(&RequestListLock);
+		// Check if it's still there in case it has been cancelled during the delay period
+		uint32 EntryIndex = RequestList.IndexOfByPredicate([Request](const FManager::FHttpRetryRequestEntry& Entry) { return Entry.Request == Request; });
+		if (EntryIndex != INDEX_NONE)
 		{
-			FHttpLogVerbosityTracker::Get().IncrementRetriedRequests();
+			FManager::FHttpRetryRequestEntry* HttpRetryRequestEntry = &RequestList[EntryIndex];
+			RetryHttpRequest(*HttpRetryRequestEntry);
 		}
-		++RequestEntry.CurrentRetryCount;
-		RequestEntry.Request->Status = FRequest::EStatus::Processing;
-	}
+	}, InDelay);
 }
 
 float FHttpRetrySystem::FManager::GetLockoutPeriodSeconds(const FHttpRetryRequestEntry& HttpRetryRequestEntry)
@@ -341,6 +433,11 @@ static FRandomStream TempRandomStream(4435261);
 
 bool FHttpRetrySystem::FManager::Update(uint32* FileCount, uint32* FailingCount, uint32* FailedCount, uint32* CompletedCount)
 {
+	if (CVarHttpRetrySystemNonGameThreadSupportEnabled.GetValueOnAnyThread())
+	{
+		return true;
+	}
+
 	QUICK_SCOPE_CYCLE_COUNTER(STAT_FHttpRetrySystem_FManager_Update);
 	LLM_SCOPE_BYTAG(HTTP);
 
@@ -411,21 +508,21 @@ bool FHttpRetrySystem::FManager::Update(uint32* FileCount, uint32* FailingCount,
 		if (HttpRetryRequestEntry->bShouldCancel)
 		{
 			UE_LOG(LogHttp, Warning, TEXT("Request cancelled on %s"), *(HttpRetryRequest->GetURL()));
-			HttpRetryRequest->Status = FHttpRetrySystem::FRequest::EStatus::Cancelled;
+			HttpRetryRequest->RetryStatus = FHttpRetrySystem::FRequest::EStatus::Cancelled;
 		}
 		else
 		{
 			if (!HasTimedOut(*HttpRetryRequestEntry, NowAbsoluteSeconds))
 			{
-				if (HttpRetryRequest->Status == FHttpRetrySystem::FRequest::EStatus::NotStarted)
+				if (HttpRetryRequest->RetryStatus == FHttpRetrySystem::FRequest::EStatus::NotStarted)
 				{
 					if (RequestStatus != EHttpRequestStatus::NotStarted)
 					{
-						HttpRetryRequest->Status = FHttpRetrySystem::FRequest::EStatus::Processing;
+						HttpRetryRequest->RetryStatus = FHttpRetrySystem::FRequest::EStatus::Processing;
 					}
 				}
 
-				if (HttpRetryRequest->Status == FHttpRetrySystem::FRequest::EStatus::Processing)
+				if (HttpRetryRequest->RetryStatus == FHttpRetrySystem::FRequest::EStatus::Processing)
 				{
 					bool forceFail = false;
 
@@ -470,7 +567,7 @@ bool FHttpRetrySystem::FManager::Update(uint32* FileCount, uint32* FailingCount,
 							}
 
 							HttpRetryRequestEntry->LockoutEndTimeAbsoluteSeconds = NowAbsoluteSeconds + LockoutPeriod;
-							HttpRetryRequest->Status = FHttpRetrySystem::FRequest::EStatus::ProcessingLockout;
+							HttpRetryRequest->RetryStatus = FHttpRetrySystem::FRequest::EStatus::ProcessingLockout;
 							
 							QUICK_SCOPE_CYCLE_COUNTER(STAT_FHttpRetrySystem_FManager_Update_OnRequestWillRetry);
 							HttpRetryRequest->OnRequestWillRetry().ExecuteIfBound(HttpRetryRequest, HttpRetryRequest->GetResponse(), LockoutPeriod);
@@ -483,7 +580,7 @@ bool FHttpRetrySystem::FManager::Update(uint32* FileCount, uint32* FailingCount,
 							{
 								++(*FailedCount);
 							}
-							HttpRetryRequest->Status = FHttpRetrySystem::FRequest::EStatus::FailedRetry;
+							HttpRetryRequest->RetryStatus = FHttpRetrySystem::FRequest::EStatus::FailedRetry;
 						}
 					}
 					else if (RequestStatus == EHttpRequestStatus::Succeeded)
@@ -498,11 +595,11 @@ bool FHttpRetrySystem::FManager::Update(uint32* FileCount, uint32* FailingCount,
 							++(*CompletedCount);
 						}
 
-						HttpRetryRequest->Status = FHttpRetrySystem::FRequest::EStatus::Succeeded;
+						HttpRetryRequest->RetryStatus = FHttpRetrySystem::FRequest::EStatus::Succeeded;
 					}
 				}
 
-				if (HttpRetryRequest->Status == FHttpRetrySystem::FRequest::EStatus::ProcessingLockout)
+				if (HttpRetryRequest->RetryStatus == FHttpRetrySystem::FRequest::EStatus::ProcessingLockout)
 				{
 					if (NowAbsoluteSeconds >= HttpRetryRequestEntry->LockoutEndTimeAbsoluteSeconds)
 					{
@@ -520,7 +617,7 @@ bool FHttpRetrySystem::FManager::Update(uint32* FileCount, uint32* FailingCount,
 			{
 				UE_LOG(LogHttp, Warning, TEXT("Timeout on retry %d: %s"), HttpRetryRequestEntry->CurrentRetryCount + 1, *(HttpRetryRequest->GetURL()));
 				bIsGreen = false;
-				HttpRetryRequest->Status = FHttpRetrySystem::FRequest::EStatus::FailedTimeout;
+				HttpRetryRequest->RetryStatus = FHttpRetrySystem::FRequest::EStatus::FailedTimeout;
 				if (FailedCount != nullptr)
 				{
 					++(*FailedCount);
@@ -531,13 +628,13 @@ bool FHttpRetrySystem::FManager::Update(uint32* FileCount, uint32* FailingCount,
 		bool bWasCompleted = false;
 		bool bWasSuccessful = false;
 
-        if (HttpRetryRequest->Status == FHttpRetrySystem::FRequest::EStatus::Cancelled ||
-            HttpRetryRequest->Status == FHttpRetrySystem::FRequest::EStatus::FailedRetry ||
-            HttpRetryRequest->Status == FHttpRetrySystem::FRequest::EStatus::FailedTimeout ||
-            HttpRetryRequest->Status == FHttpRetrySystem::FRequest::EStatus::Succeeded)
+        if (HttpRetryRequest->RetryStatus == FHttpRetrySystem::FRequest::EStatus::Cancelled ||
+            HttpRetryRequest->RetryStatus == FHttpRetrySystem::FRequest::EStatus::FailedRetry ||
+            HttpRetryRequest->RetryStatus == FHttpRetrySystem::FRequest::EStatus::FailedTimeout ||
+            HttpRetryRequest->RetryStatus == FHttpRetrySystem::FRequest::EStatus::Succeeded)
 		{
 			bWasCompleted = true;
-            bWasSuccessful = HttpRetryRequest->Status == FHttpRetrySystem::FRequest::EStatus::Succeeded;
+            bWasSuccessful = HttpRetryRequest->RetryStatus == FHttpRetrySystem::FRequest::EStatus::Succeeded;
 		}
 
 		if (bWasCompleted)
@@ -589,15 +686,12 @@ bool FHttpRetrySystem::FManager::ProcessRequest(TSharedRef<FHttpRetrySystem::FRe
 {
 	QUICK_SCOPE_CYCLE_COUNTER(STAT_FHttpRetrySystem_FManager_ProcessRequest);
 
-	bool bResult = HttpRetryRequest->HttpRequest->ProcessRequest();
-
-	if (bResult)
-	{
-		FScopeLock ScopeLock(&RequestListLock);
-		RequestList.Add(FHttpRetryRequestEntry(HttpRetryRequest));
-	}
+	FScopeLock ScopeLock(&RequestListLock);
+	RequestList.Add(FHttpRetryRequestEntry(HttpRetryRequest));
+	HttpRetryRequest->RetryStatus = FHttpRetrySystem::FRequest::EStatus::Processing;
+	HttpRetryRequest->HttpRequest->ProcessRequest();
 
-	return bResult;
+	return true;
 }
 
 void FHttpRetrySystem::FManager::CancelRequest(TSharedRef<FHttpRetrySystem::FRequest, ESPMode::ThreadSafe>& HttpRetryRequest)
@@ -626,7 +720,9 @@ void FHttpRetrySystem::FManager::CancelRequest(TSharedRef<FHttpRetrySystem::FReq
 		RetryRequestEntry.bShouldCancel = true;
 		RequestList.Add(RetryRequestEntry);
 	}
+
 	HttpRetryRequest->HttpRequest->CancelRequest();
+	HttpRetryRequest->RetryStatus = FHttpRetrySystem::FRequest::EStatus::Cancelled;
 }
 
 /* This should only be used when shutting down or suspending, to make sure 
@@ -637,13 +733,22 @@ void FHttpRetrySystem::FManager::BlockUntilFlushed(float InTimeoutSec)
 	float TimeElapsed = 0.0f;
 	uint32 FileCount, FailingCount, FailedCount, CompleteCount;
 
-	FScopeLock ScopeLock(&RequestListLock);
-
-	// TODO: Cache completed request adaptors and trigger the delegates in expected thread
-	while (RequestList.Num() > 0 && TimeElapsed < InTimeoutSec)
+	while (TimeElapsed < InTimeoutSec)
 	{
+		{
+			FScopeLock ScopeLock(&RequestListLock);
+			if (RequestList.IsEmpty())
+			{
+				break;
+			}
+		}
+
 		FHttpModule::Get().GetHttpManager().Tick(SleepInterval);
+
+PRAGMA_DISABLE_DEPRECATION_WARNINGS
 		Update(&FileCount, &FailingCount, &FailedCount, &CompleteCount);
+PRAGMA_ENABLE_DEPRECATION_WARNINGS
+
 		FPlatformProcess::Sleep(SleepInterval);
 		TimeElapsed += SleepInterval;
 	}
@@ -668,7 +773,8 @@ FHttpRetrySystem::FManager::FHttpLogVerbosityTracker::~FHttpLogVerbosityTracker(
 
 void FHttpRetrySystem::FManager::FHttpLogVerbosityTracker::IncrementRetriedRequests()
 {
-	check(IsInGameThread());
+	FScopeLock ScopeLock(&NumRetriedRequestsLock);
+
 	++NumRetriedRequests;
 	if (NumRetriedRequests == 1)
 	{
@@ -686,7 +792,8 @@ void FHttpRetrySystem::FManager::FHttpLogVerbosityTracker::IncrementRetriedReque
 
 void FHttpRetrySystem::FManager::FHttpLogVerbosityTracker::DecrementRetriedRequests()
 {
-	check(IsInGameThread());
+	FScopeLock ScopeLock(&NumRetriedRequestsLock);
+
 	--NumRetriedRequests;
 	check(NumRetriedRequests >= 0);
 	if (NumRetriedRequests == 0)
diff --git a/Engine/Source/Runtime/Online/HTTP/Private/HttpThread.cpp b/Engine/Source/Runtime/Online/HTTP/Private/HttpThread.cpp
--- a/Engine/Source/Runtime/Online/HTTP/Private/HttpThread.cpp
+++ b/Engine/Source/Runtime/Online/HTTP/Private/HttpThread.cpp
@@ -110,13 +110,10 @@ uint32 FHttpThreadBase::Run()
 
 void FHttpThreadBase::Tick()
 {
-	// Run HttpThread tasks
-	TFunction<void()> Task = nullptr;
-	while (HttpThreadQueue.Dequeue(Task))
-	{
-		check(Task);
-		Task();
-	}
+	const double AppTime = FPlatformTime::Seconds();
+	const double ElapsedTime = AppTime - LastTime;
+	LastTime = AppTime;
+	HttpThreadTick(ElapsedTime);
 }
 
 bool FHttpThreadBase::NeedsSingleThreadTick() const
@@ -140,17 +137,8 @@ void FHttpThreadBase::UpdateConfigs()
 	}
 }
 
-void FHttpThreadBase::AddHttpThreadTask(TFunction<void()>&& Task)
-{
-	if (Task)
-	{
-		HttpThreadQueue.Enqueue(MoveTemp(Task));
-	}
-}
-
 void FHttpThreadBase::HttpThreadTick(float DeltaSeconds)
 {
-	// empty
 }
 
 bool FHttpThreadBase::StartThreadedRequest(IHttpThreadedRequest* Request)
@@ -359,6 +347,21 @@ void FLegacyHttpThread::Tick()
 	}
 }
 
+void FLegacyHttpThread::AddHttpThreadTask(TFunction<void()>&& Task, float InDelay)
+{
+	Ticker.AddTicker(FTickerDelegate::CreateLambda([this, Task=MoveTemp(Task)](float) {
+		Task();
+		return false;
+	}), InDelay);
+}
+
+void FLegacyHttpThread::HttpThreadTick(float DeltaSeconds)
+{
+	FHttpThreadBase::HttpThreadTick(DeltaSeconds);
+
+	Ticker.Tick(DeltaSeconds);
+}
+
 bool FLegacyHttpThread::Init()
 {
 	ExitRequest.Set(false);
diff --git a/Engine/Source/Runtime/Online/HTTP/Private/HttpThread.h b/Engine/Source/Runtime/Online/HTTP/Private/HttpThread.h
--- a/Engine/Source/Runtime/Online/HTTP/Private/HttpThread.h
+++ b/Engine/Source/Runtime/Online/HTTP/Private/HttpThread.h
@@ -10,6 +10,7 @@
 #include "Misc/Timespan.h"
 #include "Containers/Queue.h"
 #include "Containers/SpscQueue.h"
+#include "Containers/Ticker.h"
 
 #include <atomic>
 
@@ -81,9 +82,10 @@ public:
 	/**
 	 * Add task to be ran on the http thread next tick
 	 *
-	 * @param Task The task to be ran next tick
+	 * @param Task The task to be ran
+	 * @param InDelay The delay to wait before running the task
 	 */
-	void AddHttpThreadTask(TFunction<void()>&& Task);
+	virtual void AddHttpThreadTask(TFunction<void()>&& Task, float InDelay) = 0;
 
 protected:
 
@@ -102,7 +104,6 @@ protected:
 	 */
 	virtual void CompleteThreadedRequest(IHttpThreadedRequest* Request);
 
-
 protected:
 	int32 GetRunningThreadedRequestLimit() const;
 
@@ -170,9 +171,6 @@ protected:
 	 * Added to on HTTP thread, processed then cleared on game thread (Single producer, single consumer)
 	 */
 	TSpscQueue<IHttpThreadedRequest*> CompletedThreadedRequests;
-
-	/** Queue of tasks to run on the game thread */
-	TQueue<TFunction<void()>, EQueueMode::Mpsc> HttpThreadQueue;
 };
 
 class FLegacyHttpThread	: public FHttpThreadBase
@@ -201,6 +199,9 @@ protected:
 	virtual void Stop() override;
 	//~ End FRunnable Interface
 
+	virtual void AddHttpThreadTask(TFunction<void()>&& Task, float InDelay) override;
+	virtual void HttpThreadTick(float DeltaSeconds) override;
+
 	/** signal request to stop and exit thread */
 	FThreadSafeCounter ExitRequest;
 
@@ -212,4 +213,7 @@ protected:
 	double HttpThreadIdleFrameTimeInSeconds;
 	/** Time in seconds to sleep minimally when idle, waiting for requests. */
 	double HttpThreadIdleMinimumSleepTimeInSeconds;
+
+	/* Ticker for functions to run in HTTP thread */
+	FTSTicker Ticker;
 };
diff --git a/Engine/Source/Runtime/Online/HTTP/Private/IHttpThreadedRequest.cpp b/Engine/Source/Runtime/Online/HTTP/Private/IHttpThreadedRequest.cpp
--- a/Engine/Source/Runtime/Online/HTTP/Private/IHttpThreadedRequest.cpp
+++ b/Engine/Source/Runtime/Online/HTTP/Private/IHttpThreadedRequest.cpp
@@ -4,22 +4,20 @@
 #include "HttpModule.h"
 #include "HttpManager.h"
 
-bool IHttpThreadedRequest::FinishRequestNotInHttpManager()
+void IHttpThreadedRequest::FinishRequestNotInHttpManager()
 {
 	if (IsInGameThread())
 	{
 		if (DelegateThreadPolicy == EHttpRequestDelegateThreadPolicy::CompleteOnGameThread)
 		{
 			FinishRequest();
-			return false;
 		}
 		else
 		{
 			FHttpModule::Get().GetHttpManager().AddHttpThreadTask([StrongThis = StaticCastSharedRef<IHttpThreadedRequest>(AsShared())]()
-				{
-					StrongThis->FinishRequest();
-				});
-			return true;
+			{
+				StrongThis->FinishRequest();
+			});
 		}
 	}
 	else
@@ -27,15 +25,13 @@ bool IHttpThreadedRequest::FinishRequestNotInHttpManager()
 		if (DelegateThreadPolicy == EHttpRequestDelegateThreadPolicy::CompleteOnHttpThread)
 		{
 			FinishRequest();
-			return false;
 		}
 		else
 		{
 			FHttpModule::Get().GetHttpManager().AddGameThreadTask([StrongThis = StaticCastSharedRef<IHttpThreadedRequest>(AsShared())]()
-				{
-					StrongThis->FinishRequest();
-				});
-			return true;
+			{
+				StrongThis->FinishRequest();
+			});
 		}
 	}
 }
diff --git a/Engine/Source/Runtime/Online/HTTP/Private/IHttpThreadedRequest.h b/Engine/Source/Runtime/Online/HTTP/Private/IHttpThreadedRequest.h
--- a/Engine/Source/Runtime/Online/HTTP/Private/IHttpThreadedRequest.h
+++ b/Engine/Source/Runtime/Online/HTTP/Private/IHttpThreadedRequest.h
@@ -19,8 +19,6 @@ public:
 protected:
 	/**
 	 * Finish the request when it's not in http manager
-	 * 
-	 * @return false if the request finished immediate without waiting, otherwise it needs to be finished in another thread instead of current thread
 	 */
-	bool FinishRequestNotInHttpManager();
+	void FinishRequestNotInHttpManager();
 };
diff --git a/Engine/Source/Runtime/Online/HTTP/Public/HttpManager.h b/Engine/Source/Runtime/Online/HTTP/Public/HttpManager.h
--- a/Engine/Source/Runtime/Online/HTTP/Public/HttpManager.h
+++ b/Engine/Source/Runtime/Online/HTTP/Public/HttpManager.h
@@ -232,9 +232,10 @@ public:
 	/**
 	 * Add task to be ran on the http thread next tick
 	 *
-	 * @param Task The task to be ran next tick
+	 * @param Task The task to be ran
+	 * @param InDelay The delay to wait before running the task
 	 */
-	HTTP_API void AddHttpThreadTask(TFunction<void()>&& Task);
+	HTTP_API void AddHttpThreadTask(TFunction<void()>&& Task, float InDelay = 0.0f);
 
 	/**
 	 * Set url request filter through code, instead of setting it through config.
@@ -315,4 +316,9 @@ PACKAGE_SCOPE:
 	 * Called explicitly by non-threaded requests
 	 */
 	HTTP_API void BroadcastHttpRequestCompleted(const FHttpRequestRef& Request);
+
+	/**
+	 * Access http thread of http manager for internal usage
+	 */
+	HTTP_API FHttpThreadBase* GetThread();
 };
diff --git a/Engine/Source/Runtime/Online/HTTP/Public/HttpRetrySystem.h b/Engine/Source/Runtime/Online/HTTP/Public/HttpRetrySystem.h
--- a/Engine/Source/Runtime/Online/HTTP/Public/HttpRetrySystem.h
+++ b/Engine/Source/Runtime/Online/HTTP/Public/HttpRetrySystem.h
@@ -98,7 +98,7 @@ namespace FHttpRetrySystem
 		HTTP_API virtual void CancelRequest() override;
 
 		// FRequest
-		EStatus::Type GetRetryStatus() const { return Status; }
+		EStatus::Type GetRetryStatus() const { return RetryStatus; }
 
 	protected:
 		friend class FManager;
@@ -114,13 +114,15 @@ namespace FHttpRetrySystem
 		);
 
 		void HttpOnRequestProgress(FHttpRequestPtr InHttpRequest, uint64 BytesSent, uint64 BytesRcv);
+		void HttpOnProcessRequestComplete(FHttpRequestPtr HttpRequest, FHttpResponsePtr HttpResponse, bool bSucceeded);
+		void HttpOnHeaderReceived(FHttpRequestPtr Request, const FString& HeaderName, const FString& NewHeaderValue);
 
 		/** Update our HTTP request's URL's domain from our RetryDomains */
 		void SetUrlFromRetryDomains();
 		/** Move to the next retry domain from our RetryDomains */
 		void MoveToNextRetryDomain();
 
-		EStatus::Type                        Status;
+		EStatus::Type                        RetryStatus;
 
 		FRetryLimitCountSetting              RetryLimitCountOverride;
 		FRetryTimeoutRelativeSecondsSetting  RetryTimeoutRelativeSecondsOverride;
@@ -168,7 +170,9 @@ public:
 	 *
 	 * @return                true if there are no failures or retries
 	 */
+	UE_DEPRECATED(5.4, "HttpRetrySystem::Update has been deprecated, all logic will be processed in http thread instead")
 	HTTP_API bool Update(uint32* FileCount = NULL, uint32* FailingCount = NULL, uint32* FailedCount = NULL, uint32* CompletedCount = NULL);
+
 	void SetRandomFailureRate(float Value) { RandomFailureRate = FRandomFailureRateSetting(Value); }
 	void SetDefaultRetryLimit(uint32 Value) { RetryLimitCountDefault = FRetryLimitCountSetting(Value); }
 
@@ -208,6 +212,8 @@ protected:
 
 		/** Number of requests that are in a retried state.  When this is non-zero, verbosity will be adjusted. */
 		int32 NumRetriedRequests = 0;
+		/** DecrementRetriedRequests can be called from game thread or http thread depends on the http request thread policy, make sure it's thread-safe */
+		FCriticalSection NumRetriedRequestsLock;
 		/** Verbosity to restore to when there are no requests being retried */
 		ELogVerbosity::Type OriginalVerbosity = ELogVerbosity::Error;
 		/** Config driven target verbosity to set to when requests are being retried.  NoLogging means the verbosity will not be modified. */
@@ -233,6 +239,13 @@ protected:
 	 */
 	void RetryHttpRequest(FHttpRetryRequestEntry& RequestEntry);
 
+	/**
+	 * Retry an HTTP request with delay
+	 * @param RequestToRetry request to retry
+	 * @param InDelay the delay to wait before retrying
+	 */
+	void RetryHttpRequestWithDelay(const TSharedRef<FRequest>& Request, float InDelay);
+
 	// @return number of seconds to lockout for
 	float GetLockoutPeriodSeconds(const FHttpRetryRequestEntry& HttpRetryRequestEntry);
 
