[Backout] - CL29087771
Original CL Desc
-----------------------------------------------------------------
Make retry system don't rely on Tick from game thread;
Make http tests can run with retry system enabled/disabled.

#tests Passed WebTests, also tried on Windows to reach main menu

diff --git a/Engine/Source/Programs/WebTests/Private/TestHttp.cpp b/Engine/Source/Programs/WebTests/Private/TestHttp.cpp
--- a/Engine/Source/Programs/WebTests/Private/TestHttp.cpp
+++ b/Engine/Source/Programs/WebTests/Private/TestHttp.cpp
@@ -38,7 +38,6 @@ public:
 		: WebServerIp(TEXT("127.0.0.1"))
 		, WebServerHttpPort(8000)
 		, bRunHeavyTests(false)
-		, bRetryEnabled(true)
 		, OldVerbosity(LogHttp.GetVerbosity())
 	{
 		ParseSettingsFromCommandLine();
@@ -66,7 +65,6 @@ public:
 	{
 		FParse::Value(FCommandLine::Get(), TEXT("web_server_ip"), WebServerIp);
 		FParse::Bool(FCommandLine::Get(), TEXT("run_heavy_tests"), bRunHeavyTests);
-		FParse::Bool(FCommandLine::Get(), TEXT("retry_enabled"), bRetryEnabled);
 	}
 
 	void DisableWarningsInThisTest()
@@ -85,7 +83,6 @@ public:
 	FHttpModule* HttpModule;
 
 	bool bRunHeavyTests;
-	bool bRetryEnabled;
 	ELogVerbosity::Type OldVerbosity;
 };
 
@@ -116,15 +113,57 @@ TEST_CASE_METHOD(FHttpModuleTestFixture, "Shutdown http module without issue whe
 	HttpModule->GetHttpManager().Tick(0.0f);
 }
 
-class FMockRetryManager : public FHttpRetrySystem::FManager
+class FWaitUntilQuitFromTestFixture : public FHttpModuleTestFixture
 {
 public:
-	using FHttpRetrySystem::FManager::FManager;
-	using FHttpRetrySystem::FManager::RequestList;
-	using FHttpRetrySystem::FManager::FHttpRetryRequestEntry;
-	using FHttpRetrySystem::FManager::RetryTimeoutRelativeSecondsDefault;
+	FWaitUntilQuitFromTestFixture()
+	{
+	}
+
+	~FWaitUntilQuitFromTestFixture()
+	{
+		WaitUntilQuitFromTest();
+	}
+
+	void WaitUntilQuitFromTest()
+	{
+		while (!bQuitRequested)
+		{
+			HttpModule->GetHttpManager().Tick(TickFrequency);
+			FPlatformProcess::Sleep(TickFrequency);
+		}
+	}
+
+	float TickFrequency = 1.0f / 60; /*60 FPS*/;
+	bool bQuitRequested = false;
 };
 
+TEST_CASE_METHOD(FWaitUntilQuitFromTestFixture, "Http request can be reused", HTTP_TAG)
+{
+	TSharedRef<IHttpRequest> HttpRequest = HttpModule->CreateRequest();
+	HttpRequest->SetURL(UrlToTestMethods());
+	HttpRequest->SetVerb(TEXT("POST"));
+
+	HttpRequest->OnProcessRequestComplete().BindLambda([this](FHttpRequestPtr HttpRequest, FHttpResponsePtr HttpResponse, bool bSucceeded) {
+		CHECK(bSucceeded);
+		CHECK(HttpResponse->GetResponseCode() == 200);
+
+		uint32 Chunks = 3;
+		uint32 ChunkSize = 1024;
+		HttpRequest->SetURL(UrlStreamDownload(Chunks, ChunkSize));
+		HttpRequest->SetVerb(TEXT("GET"));
+		HttpRequest->OnProcessRequestComplete().BindLambda([this, Chunks, ChunkSize](FHttpRequestPtr HttpRequest, FHttpResponsePtr HttpResponse, bool bSucceeded) {
+			CHECK(bSucceeded);
+			REQUIRE(HttpResponse != nullptr);
+			CHECK(HttpResponse->GetResponseCode() == 200);
+			CHECK(HttpResponse->GetContentLength() == Chunks * ChunkSize);
+			bQuitRequested = true;
+		});
+		HttpRequest->ProcessRequest();
+	});
+	HttpRequest->ProcessRequest();
+}
+
 class FWaitUntilCompleteHttpFixture : public FHttpModuleTestFixture
 {
 public:
@@ -132,12 +171,6 @@ public:
 	{
 		HttpModule->GetHttpManager().SetRequestAddedDelegate(FHttpManagerRequestAddedDelegate::CreateRaw(this, &FWaitUntilCompleteHttpFixture::OnRequestAdded));
 		HttpModule->GetHttpManager().SetRequestCompletedDelegate(FHttpManagerRequestCompletedDelegate::CreateRaw(this, &FWaitUntilCompleteHttpFixture::OnRequestCompleted));
-
-		if (bRetryEnabled)
-		{
-			HttpRetryManager = MakeShared<FMockRetryManager>(FHttpRetrySystem::FRetryLimitCountSetting(RetryLimitCount), FHttpRetrySystem::FRetryTimeoutRelativeSecondsSetting(/*RetryTimeoutRelativeSeconds*/));
-			HttpRetryManager->RetryTimeoutRelativeSecondsDefault = 2.0f; // Value is set so that retry system will use it at high level, will no longer wait for ever for HTTP code
-		}
 	}
 
 	~FWaitUntilCompleteHttpFixture()
@@ -171,21 +204,13 @@ public:
 		HttpModule->GetHttpManager().Tick(TickFrequency);
 	}
 
-	TSharedRef<IHttpRequest> CreateRequest()
-	{
-		return bRetryEnabled ? HttpRetryManager->CreateRequest() : HttpModule->CreateRequest();
-	}
-
 	std::atomic<int32> OngoingRequests = 0;
 	float TickFrequency = 1.0f / 60; /*60 FPS*/;
-
-	uint32 RetryLimitCount = 2;
-	TSharedPtr<FMockRetryManager> HttpRetryManager;
 };
 
 TEST_CASE_METHOD(FWaitUntilCompleteHttpFixture, "Http Methods", HTTP_TAG)
 {
-	TSharedRef<IHttpRequest> HttpRequest = CreateRequest();
+	TSharedRef<IHttpRequest> HttpRequest = HttpModule->CreateRequest();
 	CHECK(HttpRequest->GetVerb() == TEXT("GET"));
 
 	HttpRequest->SetURL(UrlToTestMethods());
@@ -220,7 +245,7 @@ TEST_CASE_METHOD(FWaitUntilCompleteHttpFixture, "Http Methods", HTTP_TAG)
 
 TEST_CASE_METHOD(FWaitUntilCompleteHttpFixture, "Get large response content without chunks", HTTP_TAG)
 {
-	TSharedRef<IHttpRequest> HttpRequest = CreateRequest();
+	TSharedRef<IHttpRequest> HttpRequest = HttpModule->CreateRequest();
 	HttpRequest->SetURL(FString::Format(TEXT("{0}/get_large_response_without_chunks/{1}/"), { *UrlHttpTests(), 1024 * 1024/*bytes_number*/}));
 	HttpRequest->SetVerb(TEXT("GET"));
 	HttpRequest->OnProcessRequestComplete().BindLambda([](FHttpRequestPtr HttpRequest, FHttpResponsePtr HttpResponse, bool bSucceeded) {
@@ -234,7 +259,7 @@ TEST_CASE_METHOD(FWaitUntilCompleteHttpFixture, "Get large response content with
 // TODO: Enable this after finding a more reliable way of simulate timeout
 //TEST_CASE_METHOD(FWaitUntilCompleteHttpFixture, "Http request connect timeout", HTTP_TAG)
 //{
-//	TSharedRef<IHttpRequest> HttpRequest = CreateRequest();
+//	TSharedRef<IHttpRequest> HttpRequest = HttpModule->CreateRequest();
 //	HttpRequest->SetURL(UrlWithInvalidPortToTestConnectTimeout());
 //	HttpRequest->SetVerb(TEXT("GET"));
 //	HttpRequest->SetTimeout(7);
@@ -257,7 +282,7 @@ TEST_CASE_METHOD(FWaitUntilCompleteHttpFixture, "Streaming http download", HTTP_
 	uint32 Chunks = 3;
 	uint32 ChunkSize = 1024*1024;
 
-	TSharedRef<IHttpRequest> HttpRequest = CreateRequest();
+	TSharedRef<IHttpRequest> HttpRequest = HttpModule->CreateRequest();
 	HttpRequest->SetURL(UrlStreamDownload(Chunks, ChunkSize));
 	HttpRequest->SetVerb(TEXT("GET"));
 
@@ -407,7 +432,7 @@ TEST_CASE_METHOD(FWaitUntilCompleteHttpFixture, "Can run parallel stream downloa
 
 	for (int i = 0; i < 3; ++i)
 	{
-		TSharedRef<IHttpRequest> HttpRequest = CreateRequest();
+		TSharedRef<IHttpRequest> HttpRequest = HttpModule->CreateRequest();
 		HttpRequest->SetURL(UrlStreamDownload(Chunks, ChunkSize));
 		HttpRequest->SetVerb(TEXT("GET"));
 		HttpRequest->OnProcessRequestComplete().BindLambda([Chunks, ChunkSize](FHttpRequestPtr HttpRequest, FHttpResponsePtr HttpResponse, bool bSucceeded) {
@@ -430,7 +455,7 @@ TEST_CASE_METHOD(FWaitUntilCompleteHttpFixture, "Can download big file exceeds 3
 	uint64 Chunks = 5 * 1024;
 	uint64 ChunkSize = 1024 * 1024;
 
-	TSharedRef<IHttpRequest> HttpRequest = CreateRequest();
+	TSharedRef<IHttpRequest> HttpRequest = HttpModule->CreateRequest();
 	HttpRequest->SetURL(UrlStreamDownload(Chunks, ChunkSize));
 	HttpRequest->SetVerb(TEXT("GET"));
 
@@ -455,7 +480,7 @@ TEST_CASE_METHOD(FWaitUntilCompleteHttpFixture, "Can download big file exceeds 3
 
 TEST_CASE_METHOD(FWaitUntilCompleteHttpFixture, "Streaming http upload from memory", HTTP_TAG)
 {
-	TSharedRef<IHttpRequest> HttpRequest = CreateRequest();
+	TSharedRef<IHttpRequest> HttpRequest = HttpModule->CreateRequest();
 	HttpRequest->SetURL(FString::Format(TEXT("{0}/streaming_upload_post"), { *UrlHttpTests() }));
 	HttpRequest->SetVerb(TEXT("POST"));
 
@@ -549,7 +574,7 @@ TEST_CASE_METHOD(FWaitUntilCompleteHttpFixture, "Can upload big file exceeds 32
 	const uint64 TotalSize = 2147483647;
 	TSharedRef<FTestHttpUploadStream> Stream = MakeShared<FTestHttpUploadStream>(TotalSize);
 
-	TSharedRef<IHttpRequest> HttpRequest = CreateRequest();
+	TSharedRef<IHttpRequest> HttpRequest = HttpModule->CreateRequest();
 	HttpRequest->SetURL(FString::Format(TEXT("{0}/streaming_upload_put"), { *UrlHttpTests() }));
 	HttpRequest->SetVerb(TEXT("PUT"));
 	HttpRequest->SetContentFromStream(Stream);
@@ -578,7 +603,7 @@ TEST_CASE_METHOD(FWaitUntilCompleteHttpFixture, "Streaming http upload from file
 	FileToWrite->Close();
 	FMemory::Free(FileData);
 
-	TSharedRef<IHttpRequest> HttpRequest = CreateRequest();
+	TSharedRef<IHttpRequest> HttpRequest = HttpModule->CreateRequest();
 	HttpRequest->SetURL(FString::Format(TEXT("{0}/streaming_upload_put"), { *UrlHttpTests() }));
 	HttpRequest->SetVerb(TEXT("PUT"));
 	HttpRequest->SetHeader(TEXT("Content-Disposition"), TEXT("attachment;filename=TestStreamUpload.dat"));
@@ -593,7 +618,7 @@ TEST_CASE_METHOD(FWaitUntilCompleteHttpFixture, "Streaming http upload from file
 
 TEST_CASE_METHOD(FWaitUntilCompleteHttpFixture, "Redirect enabled by default and can work well", HTTP_TAG)
 {
-	TSharedRef<IHttpRequest> HttpRequest = CreateRequest();
+	TSharedRef<IHttpRequest> HttpRequest = HttpModule->CreateRequest();
 	HttpRequest->SetURL(FString::Format(TEXT("{0}/redirect_from"), { *UrlHttpTests() }));
 	HttpRequest->SetVerb(TEXT("GET"));
 	HttpRequest->OnProcessRequestComplete().BindLambda([](FHttpRequestPtr HttpRequest, FHttpResponsePtr HttpResponse, bool bSucceeded) {
@@ -603,58 +628,6 @@ TEST_CASE_METHOD(FWaitUntilCompleteHttpFixture, "Redirect enabled by default and
 	HttpRequest->ProcessRequest();
 }
 
-
-class FWaitUntilQuitFromTestFixture : public FWaitUntilCompleteHttpFixture
-{
-public:
-	FWaitUntilQuitFromTestFixture()
-	{
-	}
-
-	~FWaitUntilQuitFromTestFixture()
-	{
-		WaitUntilQuitFromTest();
-	}
-
-	void WaitUntilQuitFromTest()
-	{
-		while (!bQuitRequested)
-		{
-			HttpModule->GetHttpManager().Tick(TickFrequency);
-			FPlatformProcess::Sleep(TickFrequency);
-		}
-	}
-
-	bool bQuitRequested = false;
-};
-
-TEST_CASE_METHOD(FWaitUntilQuitFromTestFixture, "Http request can be reused", HTTP_TAG)
-{
-	TSharedRef<IHttpRequest> HttpRequest = CreateRequest();
-	HttpRequest->SetURL(UrlToTestMethods());
-	HttpRequest->SetVerb(TEXT("POST"));
-
-	HttpRequest->OnProcessRequestComplete().BindLambda([this](FHttpRequestPtr HttpRequest, FHttpResponsePtr HttpResponse, bool bSucceeded) {
-		CHECK(bSucceeded);
-		CHECK(HttpResponse->GetResponseCode() == 200);
-
-		uint32 Chunks = 3;
-		uint32 ChunkSize = 1024;
-		HttpRequest->SetURL(UrlStreamDownload(Chunks, ChunkSize));
-		HttpRequest->SetVerb(TEXT("GET"));
-		HttpRequest->OnProcessRequestComplete().BindLambda([this, Chunks, ChunkSize](FHttpRequestPtr HttpRequest, FHttpResponsePtr HttpResponse, bool bSucceeded) {
-			CHECK(bSucceeded);
-			REQUIRE(HttpResponse != nullptr);
-			CHECK(HttpResponse->GetResponseCode() == 200);
-			CHECK(HttpResponse->GetContentLength() == Chunks * ChunkSize);
-			bQuitRequested = true;
-		});
-		HttpRequest->ProcessRequest();
-	});
-	HttpRequest->ProcessRequest();
-}
-
-
 // Response shared ptr should be able to be kept by user code and valid to access without http request
 class FValidateResponseDependencyFixture : public FWaitUntilCompleteHttpFixture
 {
@@ -775,7 +748,7 @@ public:
 TEST_CASE_METHOD(FWaitThreadedHttpFixture, "Http streaming download request can work in non game thread", HTTP_TAG)
 {
 	ThreadedHttpRunnable.OnRunFromThread().BindLambda([this]() {
-		TSharedRef<IHttpRequest> HttpRequest = CreateRequest();
+		TSharedRef<IHttpRequest> HttpRequest = HttpModule->CreateRequest();
 		HttpRequest->SetURL(UrlStreamDownload(3/*Chunks*/, 1024/*ChunkSize*/));
 		HttpRequest->SetVerb(TEXT("GET"));
 		HttpRequest->SetDelegateThreadPolicy(EHttpRequestDelegateThreadPolicy::CompleteOnHttpThread);
@@ -811,7 +784,7 @@ TEST_CASE_METHOD(FWaitThreadedHttpFixture, "Http download request progress callb
 {
 	std::atomic<bool> bRequestProgressTriggered = false;
 	ThreadedHttpRunnable.OnRunFromThread().BindLambda([this, &bRequestProgressTriggered]() {
-		TSharedRef<IHttpRequest> HttpRequest = CreateRequest();
+		TSharedRef<IHttpRequest> HttpRequest = HttpModule->CreateRequest();
 		HttpRequest->SetURL(UrlStreamDownload(10/*Chunks*/, 1024*1024/*ChunkSize*/));
 		HttpRequest->SetVerb(TEXT("GET"));
 
@@ -885,52 +858,37 @@ void SetupURLRequestFilter(FHttpModule* HttpModule)
 }
 }
 
-// Pre-check failed requests won't be added into http manager, so it can't rely on the requested added/completed callback in FWaitUntilCompleteHttpFixture
-TEST_CASE_METHOD(FWaitUntilQuitFromTestFixture, "Http request pre check will fail by thread policy", HTTP_TAG)
+TEST_CASE_METHOD(FWaitUntilCompleteHttpFixture, "Http request pre check will fail by thread policy", HTTP_TAG)
 {
 	DisableWarningsInThisTest();
 
 	// Pre check will fail when domain is not allowed
 	UE::TestHttp::SetupURLRequestFilter(HttpModule);
 
-	TSharedRef<IHttpRequest> HttpRequest = CreateRequest();
+	TSharedRef<IHttpRequest> HttpRequest = HttpModule->CreateRequest();
 	HttpRequest->SetVerb(TEXT("GET"));
 	HttpRequest->SetURL(UrlToTestMethods());
 
 	SECTION("on game thread")
 	{
-		HttpRequest->OnProcessRequestComplete().BindLambda([this](FHttpRequestPtr HttpRequest, FHttpResponsePtr HttpResponse, bool bSucceeded) {
+		HttpRequest->OnProcessRequestComplete().BindLambda([](FHttpRequestPtr HttpRequest, FHttpResponsePtr HttpResponse, bool bSucceeded) {
 			CHECK(IsInGameThread());
 			CHECK(!bSucceeded);
-			bQuitRequested = true;
 		});
 	}
 	SECTION("on http thread")
 	{
 		HttpRequest->SetDelegateThreadPolicy(EHttpRequestDelegateThreadPolicy::CompleteOnHttpThread);
-		HttpRequest->OnProcessRequestComplete().BindLambda([this](FHttpRequestPtr HttpRequest, FHttpResponsePtr HttpResponse, bool bSucceeded) {
+		HttpRequest->OnProcessRequestComplete().BindLambda([](FHttpRequestPtr HttpRequest, FHttpResponsePtr HttpResponse, bool bSucceeded) {
 			CHECK(!IsInGameThread());
 			CHECK(!bSucceeded);
-			bQuitRequested = true;
 		});
 	}
 
 	HttpRequest->ProcessRequest();
 }
 
-class FWaitUnitilQuitFromTestThreadedFixture : public FWaitUntilQuitFromTestFixture
-{
-public:
-	~FWaitUnitilQuitFromTestThreadedFixture()
-	{
-		WaitUntilQuitFromTest();
-	}
-
-	FThreadedHttpRunnable ThreadedHttpRunnable;
-};
-
-// Pre-check failed requests won't be added into http manager, so it can't rely on the requested added/completed callback in FWaitUntilCompleteHttpFixture
-TEST_CASE_METHOD(FWaitUnitilQuitFromTestThreadedFixture, "Threaded http request pre check will fail by thread policy", HTTP_TAG)
+TEST_CASE_METHOD(FWaitThreadedHttpFixture, "Threaded http request pre check will fail by thread policy", HTTP_TAG)
 {
 	DisableWarningsInThisTest();
 
@@ -938,25 +896,23 @@ TEST_CASE_METHOD(FWaitUnitilQuitFromTestThreadedFixture, "Threaded http request
 		// Pre check will fail when domain is not allowed
 		UE::TestHttp::SetupURLRequestFilter(HttpModule);
 
-		TSharedRef<IHttpRequest> HttpRequest = CreateRequest();
+		TSharedRef<IHttpRequest> HttpRequest = HttpModule->CreateRequest();
 		HttpRequest->SetVerb(TEXT("GET"));
 		HttpRequest->SetURL(UrlToTestMethods());
 
 		SECTION("on game thread")
 		{
-			HttpRequest->OnProcessRequestComplete().BindLambda([this](FHttpRequestPtr HttpRequest, FHttpResponsePtr HttpResponse, bool bSucceeded) {
+			HttpRequest->OnProcessRequestComplete().BindLambda([](FHttpRequestPtr HttpRequest, FHttpResponsePtr HttpResponse, bool bSucceeded) {
 				CHECK(IsInGameThread());
 				CHECK(!bSucceeded);
-				bQuitRequested = true;
 			});
 		}
 		SECTION("on http thread")
 		{
 			HttpRequest->SetDelegateThreadPolicy(EHttpRequestDelegateThreadPolicy::CompleteOnHttpThread);
-			HttpRequest->OnProcessRequestComplete().BindLambda([this](FHttpRequestPtr HttpRequest, FHttpResponsePtr HttpResponse, bool bSucceeded) {
+			HttpRequest->OnProcessRequestComplete().BindLambda([](FHttpRequestPtr HttpRequest, FHttpResponsePtr HttpResponse, bool bSucceeded) {
 				CHECK(!IsInGameThread());
 				CHECK(!bSucceeded);
-				bQuitRequested = true;
 			});
 		}
 
@@ -970,7 +926,7 @@ TEST_CASE_METHOD(FWaitUntilCompleteHttpFixture, "Cancel http request connect bef
 {
 	DisableWarningsInThisTest();
 
-	TSharedRef<IHttpRequest, ESPMode::ThreadSafe> HttpRequest = CreateRequest();
+	TSharedRef<IHttpRequest, ESPMode::ThreadSafe> HttpRequest = HttpModule->CreateRequest();
 	HttpRequest->SetURL(UrlWithInvalidPortToTestConnectTimeout());
 	HttpRequest->SetVerb(TEXT("GET"));
 	HttpRequest->SetTimeout(7);
@@ -987,48 +943,107 @@ TEST_CASE_METHOD(FWaitUntilCompleteHttpFixture, "Cancel http request connect bef
 	HttpRequest->CancelRequest();
 }
 
-class FThreadedBatchRequestsFixture : public FWaitThreadedHttpFixture
+class FMockRetryManager : public FHttpRetrySystem::FManager
+{
+public:
+	using FHttpRetrySystem::FManager::FManager;
+	using FHttpRetrySystem::FManager::RequestList;
+	using FHttpRetrySystem::FManager::FHttpRetryRequestEntry;
+	using FHttpRetrySystem::FManager::RetryTimeoutRelativeSecondsDefault;
+};
+
+class FRetryHttpFixture : public FWaitUntilCompleteHttpFixture
 {
 public:
-	void LaunchBatchRequests(uint32 BatchSize)
+	FRetryHttpFixture()
 	{
-		for (uint32 i = 0; i < BatchSize; ++i)
-		{
-			TSharedRef<IHttpRequest> HttpRequest = CreateRequest();
-			HttpRequest->SetURL(UrlStreamDownload(3, 1024*1024));
-			HttpRequest->SetVerb(TEXT("GET"));
-			HttpRequest->ProcessRequest();
-		}
+		HttpRetryManager = MakeShared<FMockRetryManager>(FHttpRetrySystem::FRetryLimitCountSetting(RetryLimitCount), FHttpRetrySystem::FRetryTimeoutRelativeSecondsSetting(/*RetryTimeoutRelativeSeconds*/));
+		HttpRetryManager->RetryTimeoutRelativeSecondsDefault = 2.0f; // Value is set so that retry system will use it at high level, will no longer wait for ever for HTTP code
 	}
 
-	void BlockUntilFlushed()
+	~FRetryHttpFixture()
 	{
-		if (bRetryEnabled)
+		// Make sure not destroy retry manager before waiting in FWaitUntilCompleteHttpFixture
+		while (HttpRetryManager->RequestList.Num())
 		{
-			HttpRetryManager->BlockUntilFlushed(5.0);
+			HttpRetryManager->Update();
+			HttpModule->GetHttpManager().Tick(TickFrequency);
+			FPlatformProcess::Sleep(TickFrequency);
 		}
-		else
+	}
+
+	TSharedPtr<FMockRetryManager> HttpRetryManager;
+	uint32 RetryLimitCount = 2;
+
+	FThreadedHttpRunnable ThreadedHttpRunnable;
+};
+
+// TODO: Currently broken, will fix this in retry system
+//TEST_CASE_METHOD(FRetryHttpFixture, "Http request adaptor in retry system can work in non game thread", HTTP_TAG)
+//{
+//	ThreadedHttpRunnable.OnRunFromThread().BindLambda([this]() {
+//		TSharedRef<IHttpRequest> HttpRequest = HttpRetryManager->CreateRequest();
+//		HttpRequest->SetURL(UrlStreamDownload(3/*Chunks*/, 1024/*ChunkSize*/));
+//		HttpRequest->SetVerb(TEXT("GET"));
+//		HttpRequest->SetDelegateThreadPolicy(EHttpRequestDelegateThreadPolicy::CompleteOnHttpThread);
+//
+//		class FTestHttpReceiveStream final : public FArchive
+//		{
+//		public:
+//			virtual void Serialize(void* V, int64 Length) override
+//			{
+//				// No matter what's the thread policy, Serialize always get called in http thread.
+//				CHECK(!IsInGameThread());
+//			}
+//		};
+//		CHECK(HttpRequest->SetResponseBodyReceiveStream(MakeShared<FTestHttpReceiveStream>()));
+//
+//		HttpRequest->OnProcessRequestComplete().BindLambda([this](FHttpRequestPtr HttpRequest, FHttpResponsePtr HttpResponse, bool bSucceeded) {
+//			// EHttpRequestDelegateThreadPolicy::CompleteOnHttpThread was used, so not in game thread here
+//			CHECK(!IsInGameThread());
+//			CHECK(bSucceeded);
+//			CHECK(HttpResponse->GetResponseCode() == 200);
+//			ThreadedHttpRunnable.UnblockGameThread();
+//		});
+//
+//		HttpRequest->ProcessRequest();
+//	});
+//
+//	ThreadedHttpRunnable.StartTestHttpThread(true/*bBlockGameThread*/);
+//}
+
+class FRetryHttpManagerThreadedRequestsFixture : public FRetryHttpFixture
+{
+public:
+	void LaunchDownloadRequests()
+	{
+		for (int32 i = 0; i < 10; ++i)
 		{
-			HttpModule->GetHttpManager().Flush(EHttpFlushReason::Default);
+			TSharedRef<IHttpRequest> HttpRequest = HttpRetryManager->CreateRequest();
+			HttpRequest->SetURL(UrlStreamDownload(3, 1024*1024));
+			HttpRequest->SetVerb(TEXT("GET"));
+			HttpRequest->ProcessRequest();
 		}
 	}
 };
 
-TEST_CASE_METHOD(FThreadedBatchRequestsFixture, "Retry manager and http manager is thread safe", HTTP_TAG)
+TEST_CASE_METHOD(FRetryHttpManagerThreadedRequestsFixture, "Retry manager is thread safe", HTTP_TAG)
 {
 	DisableWarningsInThisTest();
 
 	ThreadedHttpRunnable.OnRunFromThread().BindLambda([this]() {
-		LaunchBatchRequests(10);
-		BlockUntilFlushed();
+		LaunchDownloadRequests();
+		HttpRetryManager->BlockUntilFlushed(5.0);
 	});
 	ThreadedHttpRunnable.StartTestHttpThread(false/*bBlockGameThread*/);
 
-	LaunchBatchRequests(10);
-	BlockUntilFlushed();
+	LaunchDownloadRequests();
+	HttpRetryManager->BlockUntilFlushed(5.0);
 }
 
 
 // TODO: Add cancel test, with multiple cancel calls
 
+// TODO: Add a test case to make sure once Request completed and destroyed, if the user code still keeps the shared ptr of Response, any call to Response shouldn't have dependency to Request
+
 // TODO: Add test case to validate header received callback can be received in http/game thread, and can be received before the request complete
diff --git a/Engine/Source/Runtime/Analytics/AnalyticsET/Private/IAnalyticsProviderET.cpp b/Engine/Source/Runtime/Analytics/AnalyticsET/Private/IAnalyticsProviderET.cpp
--- a/Engine/Source/Runtime/Analytics/AnalyticsET/Private/IAnalyticsProviderET.cpp
+++ b/Engine/Source/Runtime/Analytics/AnalyticsET/Private/IAnalyticsProviderET.cpp
@@ -314,9 +314,7 @@ bool FAnalyticsProviderET::Tick(float DeltaSeconds)
 {
 	QUICK_SCOPE_CYCLE_COUNTER(STAT_FAnalyticsProviderET_Tick);
 
-PRAGMA_DISABLE_DEPRECATION_WARNINGS
 	HttpRetryManager->Update();
-PRAGMA_ENABLE_DEPRECATION_WARNINGS
 
 	// hold a lock the entire time here because we're making several calls to the event cache that we need to be consistent when we decide to flush.
 	// With more care, we can likely avoid holding this lock the entire time.
diff --git a/Engine/Source/Runtime/Online/HTTP/Private/Curl/CurlHttp.cpp b/Engine/Source/Runtime/Online/HTTP/Private/Curl/CurlHttp.cpp
--- a/Engine/Source/Runtime/Online/HTTP/Private/Curl/CurlHttp.cpp
+++ b/Engine/Source/Runtime/Online/HTTP/Private/Curl/CurlHttp.cpp
@@ -978,8 +978,7 @@ bool FCurlHttpRequest::ProcessRequest()
 
 	if (!PreCheck() || !SetupRequest())
 	{
-		FinishRequestNotInHttpManager();
-		return false;
+		return FinishRequestNotInHttpManager();
 	}
 
 	// Clear the info cache log so we don't output messages from previous requests when reusing/retrying a request
@@ -1261,7 +1260,7 @@ void FCurlHttpRequest::FinishRequest()
 		{
 			UE_LOG(LogHttp, Warning, TEXT("%p: request failed, libcurl multi error: %d (%s)"), this, (int32)CurlAddToMultiResult, ANSI_TO_TCHAR(curl_multi_strerror(CurlAddToMultiResult)));
 		}
-		else if (CurlCompletionResult != CURLE_OK)
+		else
 		{
 			UE_LOG(LogHttp, Warning, TEXT("%p: request failed, libcurl error: %d (%s)"), this, (int32)CurlCompletionResult, ANSI_TO_TCHAR(curl_easy_strerror(CurlCompletionResult)));
 		}
diff --git a/Engine/Source/Runtime/Online/HTTP/Private/EventLoopHttpThread.cpp b/Engine/Source/Runtime/Online/HTTP/Private/EventLoopHttpThread.cpp
--- a/Engine/Source/Runtime/Online/HTTP/Private/EventLoopHttpThread.cpp
+++ b/Engine/Source/Runtime/Online/HTTP/Private/EventLoopHttpThread.cpp
@@ -159,10 +159,4 @@ void FEventLoopHttpThread::ResetTickTimer()
 	},
 	FTimespan::FromSeconds(FHttpModule::Get().GetHttpEventLoopThreadTickIntervalInSeconds()),
 	true /* repeat */);
-}
-
-void FEventLoopHttpThread::AddHttpThreadTask(TFunction<void()>&& Task, float InDelay)
-{
-	UE::EventLoop::IEventLoop& EventLoop = GetEventLoopChecked();
-	EventLoop.SetTimer(Task, FTimespan::FromSeconds(InDelay));
 }
\ No newline at end of file
diff --git a/Engine/Source/Runtime/Online/HTTP/Private/EventLoopHttpThread.h b/Engine/Source/Runtime/Online/HTTP/Private/EventLoopHttpThread.h
--- a/Engine/Source/Runtime/Online/HTTP/Private/EventLoopHttpThread.h
+++ b/Engine/Source/Runtime/Online/HTTP/Private/EventLoopHttpThread.h
@@ -40,8 +40,6 @@ protected:
 
 	void ResetTickTimer();
 
-	virtual void AddHttpThreadTask(TFunction<void()>&& Task, float InDelay) override;
-
 protected:
 	UE::EventLoop::FTimerHandle RequestTickTimer;
 };
diff --git a/Engine/Source/Runtime/Online/HTTP/Private/HttpManager.cpp b/Engine/Source/Runtime/Online/HTTP/Private/HttpManager.cpp
--- a/Engine/Source/Runtime/Online/HTTP/Private/HttpManager.cpp
+++ b/Engine/Source/Runtime/Online/HTTP/Private/HttpManager.cpp
@@ -284,10 +284,10 @@ void FHttpManager::AddGameThreadTask(TFunction<void()>&& Task)
 	}
 }
 
-void FHttpManager::AddHttpThreadTask(TFunction<void()>&& Task, float InDelay)
+void FHttpManager::AddHttpThreadTask(TFunction<void()>&& Task)
 {
 	check(Thread);
-	Thread->AddHttpThreadTask(MoveTemp(Task), InDelay);
+	Thread->AddHttpThreadTask(MoveTemp(Task));
 }
 
 FHttpThreadBase* FHttpManager::CreateHttpThread()
@@ -547,8 +547,3 @@ void FHttpManager::BroadcastHttpRequestCompleted(const FHttpRequestRef& Request)
 {
 	RequestCompletedDelegate.ExecuteIfBound(Request);
 }
-
-FHttpThreadBase* FHttpManager::GetThread()
-{
-	return Thread;
-}
diff --git a/Engine/Source/Runtime/Online/HTTP/Private/HttpRetrySystem.cpp b/Engine/Source/Runtime/Online/HTTP/Private/HttpRetrySystem.cpp
--- a/Engine/Source/Runtime/Online/HTTP/Private/HttpRetrySystem.cpp
+++ b/Engine/Source/Runtime/Online/HTTP/Private/HttpRetrySystem.cpp
@@ -1,7 +1,6 @@
 // Copyright Epic Games, Inc. All Rights Reserved.
 
 #include "HttpRetrySystem.h"
-#include "HAL/ConsoleManager.h"
 #include "HAL/PlatformTime.h"
 #include "HAL/PlatformProcess.h"
 #include "HAL/LowLevelMemTracker.h"
@@ -9,17 +8,10 @@
 #include "HttpModule.h"
 #include "Http.h"
 #include "HttpManager.h"
-#include "HttpThread.h"
 #include "Stats/Stats.h"
 #include "Misc/ConfigCacheIni.h"
 #include "Misc/CoreDelegates.h"
 
-TAutoConsoleVariable<bool> CVarHttpRetrySystemNonGameThreadSupportEnabled(
-	TEXT("Http.RetrySystemNonGameThreadSupportEnabled"),
-	true,
-	TEXT("Enable retry system non-game thread support")
-);
-
 LLM_DEFINE_TAG(HTTP);
 
 namespace FHttpRetrySystem
@@ -79,7 +71,7 @@ FHttpRetrySystem::FRequest::FRequest(
 	const FHttpRetrySystem::FRetryDomainsPtr& InRetryDomains
 	)
     : FHttpRequestAdapterBase(HttpRequest)
-    , RetryStatus(FHttpRetrySystem::FRequest::EStatus::NotStarted)
+    , Status(FHttpRetrySystem::FRequest::EStatus::NotStarted)
     , RetryLimitCountOverride(InRetryLimitCountOverride)
     , RetryTimeoutRelativeSecondsOverride(InRetryTimeoutRelativeSecondsOverride)
 	, RetryResponseCodes(InRetryResponseCodes)
@@ -117,8 +109,6 @@ bool FHttpRetrySystem::FRequest::ProcessRequest()
 	}
 
 	HttpRequest->OnRequestProgress64().BindThreadSafeSP(RetryRequest, &FHttpRetrySystem::FRequest::HttpOnRequestProgress);
-	HttpRequest->OnProcessRequestComplete().BindThreadSafeSP(RetryRequest, &FHttpRetrySystem::FRequest::HttpOnProcessRequestComplete);
-	HttpRequest->OnHeaderReceived().BindThreadSafeSP(RetryRequest, &FHttpRetrySystem::FRequest::HttpOnHeaderReceived);
 
 	return RetryManager.ProcessRequest(RetryRequest);
 }
@@ -160,83 +150,11 @@ PRAGMA_ENABLE_DEPRECATION_WARNINGS
 	OnRequestProgress64().ExecuteIfBound(AsShared(), BytesSent, BytesRcv);
 }
 
-void FHttpRetrySystem::FRequest::HttpOnProcessRequestComplete(FHttpRequestPtr InHttpRequest, FHttpResponsePtr HttpResponse, bool bSucceeded)
-{
-	if (!CVarHttpRetrySystemNonGameThreadSupportEnabled.GetValueOnAnyThread())
-	{
-		return;
-	}
-
-	TSharedRef<FRequest> SelfPtr = StaticCastSharedRef<FRequest>(AsShared()); // In case no ref after removing from RetryManager
-
-	FScopeLock ScopeLock(&RetryManager.RequestListLock);
-
-	uint32 EntryIndex = RetryManager.RequestList.IndexOfByPredicate([this](const FManager::FHttpRetryRequestEntry& Entry) { return Entry.Request == AsShared(); });
-	check(EntryIndex != INDEX_NONE);
-	FManager::FHttpRetryRequestEntry* HttpRetryRequestEntry = &RetryManager.RequestList[EntryIndex];
-
-	if (RetryStatus == FHttpRetrySystem::FRequest::EStatus::Cancelled)
-	{
-		// Do nothing here
-	}
-	if (GetStatus() == EHttpRequestStatus::Failed_ConnectionError || GetStatus() == EHttpRequestStatus::Failed)
-	{
-		if (GetStatus() == EHttpRequestStatus::Failed_ConnectionError && RetryDomains.IsValid())
-		{
-			MoveToNextRetryDomain();
-		}
-
-		if (RetryManager.ShouldRetry(*HttpRetryRequestEntry) && RetryManager.CanRetry(*HttpRetryRequestEntry))
-		{
-			const double NowAbsoluteSeconds = FPlatformTime::Seconds();
-			float LockoutPeriod = RetryManager.GetLockoutPeriodSeconds(*HttpRetryRequestEntry);
-
-			RetryStatus = FHttpRetrySystem::FRequest::EStatus::ProcessingLockout;
-
-			{
-				QUICK_SCOPE_CYCLE_COUNTER(STAT_FHttpRetrySystem_FManager_Update_OnRequestWillRetry);
-				OnRequestWillRetry().ExecuteIfBound(HttpRetryRequestEntry->Request, GetResponse(), LockoutPeriod);
-			}
-
-			RetryManager.RetryHttpRequestWithDelay(HttpRetryRequestEntry->Request, LockoutPeriod);
-			return;
-		}
-
-		RetryStatus = FHttpRetrySystem::FRequest::EStatus::FailedRetry;
-	}
-	else
-	{
-		RetryStatus = FHttpRetrySystem::FRequest::EStatus::Succeeded;
-	}
-
-	if (HttpRetryRequestEntry->CurrentRetryCount > 0)
-	{
-		FHttpRetrySystem::FManager::FHttpLogVerbosityTracker::Get().DecrementRetriedRequests();
-	}
-
-	RetryManager.RequestList.RemoveAtSwap(EntryIndex);
-
-	OnProcessRequestComplete().ExecuteIfBound(SelfPtr, HttpResponse, bSucceeded);
-}
-
-void FHttpRetrySystem::FRequest::HttpOnHeaderReceived(FHttpRequestPtr Request, const FString& HeaderName, const FString& NewHeaderValue)
-{
-	if (!CVarHttpRetrySystemNonGameThreadSupportEnabled.GetValueOnAnyThread())
-	{
-		return;
-	}
-
-	TSharedRef<FRequest> SelfPtr = StaticCastSharedRef<FRequest>(AsShared());
-	OnHeaderReceived().ExecuteIfBound(SelfPtr, HeaderName, NewHeaderValue);
-}
-
 FHttpRetrySystem::FManager::FManager(const FRetryLimitCountSetting& InRetryLimitCountDefault, const FRetryTimeoutRelativeSecondsSetting& InRetryTimeoutRelativeSecondsDefault)
     : RandomFailureRate(FRandomFailureRateSetting())
     , RetryLimitCountDefault(InRetryLimitCountDefault)
 	, RetryTimeoutRelativeSecondsDefault(InRetryTimeoutRelativeSecondsDefault)
-{
-	check(FHttpModule::Get().GetHttpManager().GetThread());
-}
+{}
 
 FHttpRetrySystem::FManager::~FManager()
 {
@@ -374,28 +292,18 @@ bool FHttpRetrySystem::FManager::HasTimedOut(const FHttpRetryRequestEntry& HttpR
 void FHttpRetrySystem::FManager::RetryHttpRequest(FHttpRetryRequestEntry& RequestEntry)
 {
 	// if this fails the HttpRequest's state will be failed which will cause the retry logic to kick(as expected)
-	if (RequestEntry.CurrentRetryCount == 0)
+	const bool bProcessRequestSuccess = RequestEntry.Request->HttpRequest->ProcessRequest();
+	if (bProcessRequestSuccess)
 	{
-		FHttpLogVerbosityTracker::Get().IncrementRetriedRequests();
-	}
-	++RequestEntry.CurrentRetryCount;
-	RequestEntry.Request->RetryStatus = FRequest::EStatus::Processing;
-	UE_LOG(LogHttp, Warning, TEXT("Retry %d on %s"), RequestEntry.CurrentRetryCount, *(RequestEntry.Request->GetURL()));
-	RequestEntry.Request->HttpRequest->ProcessRequest();
-}
+		UE_LOG(LogHttp, Warning, TEXT("Retry %d on %s"), RequestEntry.CurrentRetryCount + 1, *(RequestEntry.Request->GetURL()));
 
-void FHttpRetrySystem::FManager::RetryHttpRequestWithDelay(const TSharedRef<FRequest>& Request, float InDelay)
-{
-	FHttpModule::Get().GetHttpManager().AddHttpThreadTask([this, Request]() {
-		FScopeLock ScopeLock(&RequestListLock);
-		// Check if it's still there in case it has been cancelled during the delay period
-		uint32 EntryIndex = RequestList.IndexOfByPredicate([Request](const FManager::FHttpRetryRequestEntry& Entry) { return Entry.Request == Request; });
-		if (EntryIndex != INDEX_NONE)
+		if (RequestEntry.CurrentRetryCount == 0)
 		{
-			FManager::FHttpRetryRequestEntry* HttpRetryRequestEntry = &RequestList[EntryIndex];
-			RetryHttpRequest(*HttpRetryRequestEntry);
+			FHttpLogVerbosityTracker::Get().IncrementRetriedRequests();
 		}
-	}, InDelay);
+		++RequestEntry.CurrentRetryCount;
+		RequestEntry.Request->Status = FRequest::EStatus::Processing;
+	}
 }
 
 float FHttpRetrySystem::FManager::GetLockoutPeriodSeconds(const FHttpRetryRequestEntry& HttpRetryRequestEntry)
@@ -433,11 +341,6 @@ static FRandomStream TempRandomStream(4435261);
 
 bool FHttpRetrySystem::FManager::Update(uint32* FileCount, uint32* FailingCount, uint32* FailedCount, uint32* CompletedCount)
 {
-	if (CVarHttpRetrySystemNonGameThreadSupportEnabled.GetValueOnAnyThread())
-	{
-		return true;
-	}
-
 	QUICK_SCOPE_CYCLE_COUNTER(STAT_FHttpRetrySystem_FManager_Update);
 	LLM_SCOPE_BYTAG(HTTP);
 
@@ -508,21 +411,21 @@ bool FHttpRetrySystem::FManager::Update(uint32* FileCount, uint32* FailingCount,
 		if (HttpRetryRequestEntry->bShouldCancel)
 		{
 			UE_LOG(LogHttp, Warning, TEXT("Request cancelled on %s"), *(HttpRetryRequest->GetURL()));
-			HttpRetryRequest->RetryStatus = FHttpRetrySystem::FRequest::EStatus::Cancelled;
+			HttpRetryRequest->Status = FHttpRetrySystem::FRequest::EStatus::Cancelled;
 		}
 		else
 		{
 			if (!HasTimedOut(*HttpRetryRequestEntry, NowAbsoluteSeconds))
 			{
-				if (HttpRetryRequest->RetryStatus == FHttpRetrySystem::FRequest::EStatus::NotStarted)
+				if (HttpRetryRequest->Status == FHttpRetrySystem::FRequest::EStatus::NotStarted)
 				{
 					if (RequestStatus != EHttpRequestStatus::NotStarted)
 					{
-						HttpRetryRequest->RetryStatus = FHttpRetrySystem::FRequest::EStatus::Processing;
+						HttpRetryRequest->Status = FHttpRetrySystem::FRequest::EStatus::Processing;
 					}
 				}
 
-				if (HttpRetryRequest->RetryStatus == FHttpRetrySystem::FRequest::EStatus::Processing)
+				if (HttpRetryRequest->Status == FHttpRetrySystem::FRequest::EStatus::Processing)
 				{
 					bool forceFail = false;
 
@@ -567,7 +470,7 @@ bool FHttpRetrySystem::FManager::Update(uint32* FileCount, uint32* FailingCount,
 							}
 
 							HttpRetryRequestEntry->LockoutEndTimeAbsoluteSeconds = NowAbsoluteSeconds + LockoutPeriod;
-							HttpRetryRequest->RetryStatus = FHttpRetrySystem::FRequest::EStatus::ProcessingLockout;
+							HttpRetryRequest->Status = FHttpRetrySystem::FRequest::EStatus::ProcessingLockout;
 							
 							QUICK_SCOPE_CYCLE_COUNTER(STAT_FHttpRetrySystem_FManager_Update_OnRequestWillRetry);
 							HttpRetryRequest->OnRequestWillRetry().ExecuteIfBound(HttpRetryRequest, HttpRetryRequest->GetResponse(), LockoutPeriod);
@@ -580,7 +483,7 @@ bool FHttpRetrySystem::FManager::Update(uint32* FileCount, uint32* FailingCount,
 							{
 								++(*FailedCount);
 							}
-							HttpRetryRequest->RetryStatus = FHttpRetrySystem::FRequest::EStatus::FailedRetry;
+							HttpRetryRequest->Status = FHttpRetrySystem::FRequest::EStatus::FailedRetry;
 						}
 					}
 					else if (RequestStatus == EHttpRequestStatus::Succeeded)
@@ -595,11 +498,11 @@ bool FHttpRetrySystem::FManager::Update(uint32* FileCount, uint32* FailingCount,
 							++(*CompletedCount);
 						}
 
-						HttpRetryRequest->RetryStatus = FHttpRetrySystem::FRequest::EStatus::Succeeded;
+						HttpRetryRequest->Status = FHttpRetrySystem::FRequest::EStatus::Succeeded;
 					}
 				}
 
-				if (HttpRetryRequest->RetryStatus == FHttpRetrySystem::FRequest::EStatus::ProcessingLockout)
+				if (HttpRetryRequest->Status == FHttpRetrySystem::FRequest::EStatus::ProcessingLockout)
 				{
 					if (NowAbsoluteSeconds >= HttpRetryRequestEntry->LockoutEndTimeAbsoluteSeconds)
 					{
@@ -617,7 +520,7 @@ bool FHttpRetrySystem::FManager::Update(uint32* FileCount, uint32* FailingCount,
 			{
 				UE_LOG(LogHttp, Warning, TEXT("Timeout on retry %d: %s"), HttpRetryRequestEntry->CurrentRetryCount + 1, *(HttpRetryRequest->GetURL()));
 				bIsGreen = false;
-				HttpRetryRequest->RetryStatus = FHttpRetrySystem::FRequest::EStatus::FailedTimeout;
+				HttpRetryRequest->Status = FHttpRetrySystem::FRequest::EStatus::FailedTimeout;
 				if (FailedCount != nullptr)
 				{
 					++(*FailedCount);
@@ -628,13 +531,13 @@ bool FHttpRetrySystem::FManager::Update(uint32* FileCount, uint32* FailingCount,
 		bool bWasCompleted = false;
 		bool bWasSuccessful = false;
 
-        if (HttpRetryRequest->RetryStatus == FHttpRetrySystem::FRequest::EStatus::Cancelled ||
-            HttpRetryRequest->RetryStatus == FHttpRetrySystem::FRequest::EStatus::FailedRetry ||
-            HttpRetryRequest->RetryStatus == FHttpRetrySystem::FRequest::EStatus::FailedTimeout ||
-            HttpRetryRequest->RetryStatus == FHttpRetrySystem::FRequest::EStatus::Succeeded)
+        if (HttpRetryRequest->Status == FHttpRetrySystem::FRequest::EStatus::Cancelled ||
+            HttpRetryRequest->Status == FHttpRetrySystem::FRequest::EStatus::FailedRetry ||
+            HttpRetryRequest->Status == FHttpRetrySystem::FRequest::EStatus::FailedTimeout ||
+            HttpRetryRequest->Status == FHttpRetrySystem::FRequest::EStatus::Succeeded)
 		{
 			bWasCompleted = true;
-            bWasSuccessful = HttpRetryRequest->RetryStatus == FHttpRetrySystem::FRequest::EStatus::Succeeded;
+            bWasSuccessful = HttpRetryRequest->Status == FHttpRetrySystem::FRequest::EStatus::Succeeded;
 		}
 
 		if (bWasCompleted)
@@ -686,12 +589,15 @@ bool FHttpRetrySystem::FManager::ProcessRequest(TSharedRef<FHttpRetrySystem::FRe
 {
 	QUICK_SCOPE_CYCLE_COUNTER(STAT_FHttpRetrySystem_FManager_ProcessRequest);
 
-	FScopeLock ScopeLock(&RequestListLock);
-	RequestList.Add(FHttpRetryRequestEntry(HttpRetryRequest));
-	HttpRetryRequest->RetryStatus = FHttpRetrySystem::FRequest::EStatus::Processing;
-	HttpRetryRequest->HttpRequest->ProcessRequest();
+	bool bResult = HttpRetryRequest->HttpRequest->ProcessRequest();
+
+	if (bResult)
+	{
+		FScopeLock ScopeLock(&RequestListLock);
+		RequestList.Add(FHttpRetryRequestEntry(HttpRetryRequest));
+	}
 
-	return true;
+	return bResult;
 }
 
 void FHttpRetrySystem::FManager::CancelRequest(TSharedRef<FHttpRetrySystem::FRequest, ESPMode::ThreadSafe>& HttpRetryRequest)
@@ -720,9 +626,7 @@ void FHttpRetrySystem::FManager::CancelRequest(TSharedRef<FHttpRetrySystem::FReq
 		RetryRequestEntry.bShouldCancel = true;
 		RequestList.Add(RetryRequestEntry);
 	}
-
 	HttpRetryRequest->HttpRequest->CancelRequest();
-	HttpRetryRequest->RetryStatus = FHttpRetrySystem::FRequest::EStatus::Cancelled;
 }
 
 /* This should only be used when shutting down or suspending, to make sure 
@@ -733,22 +637,13 @@ void FHttpRetrySystem::FManager::BlockUntilFlushed(float InTimeoutSec)
 	float TimeElapsed = 0.0f;
 	uint32 FileCount, FailingCount, FailedCount, CompleteCount;
 
-	while (TimeElapsed < InTimeoutSec)
-	{
-		{
-			FScopeLock ScopeLock(&RequestListLock);
-			if (RequestList.IsEmpty())
-			{
-				break;
-			}
-		}
+	FScopeLock ScopeLock(&RequestListLock);
 
+	// TODO: Cache completed request adaptors and trigger the delegates in expected thread
+	while (RequestList.Num() > 0 && TimeElapsed < InTimeoutSec)
+	{
 		FHttpModule::Get().GetHttpManager().Tick(SleepInterval);
-
-PRAGMA_DISABLE_DEPRECATION_WARNINGS
 		Update(&FileCount, &FailingCount, &FailedCount, &CompleteCount);
-PRAGMA_ENABLE_DEPRECATION_WARNINGS
-
 		FPlatformProcess::Sleep(SleepInterval);
 		TimeElapsed += SleepInterval;
 	}
@@ -773,8 +668,7 @@ FHttpRetrySystem::FManager::FHttpLogVerbosityTracker::~FHttpLogVerbosityTracker(
 
 void FHttpRetrySystem::FManager::FHttpLogVerbosityTracker::IncrementRetriedRequests()
 {
-	FScopeLock ScopeLock(&NumRetriedRequestsLock);
-
+	check(IsInGameThread());
 	++NumRetriedRequests;
 	if (NumRetriedRequests == 1)
 	{
@@ -792,8 +686,7 @@ void FHttpRetrySystem::FManager::FHttpLogVerbosityTracker::IncrementRetriedReque
 
 void FHttpRetrySystem::FManager::FHttpLogVerbosityTracker::DecrementRetriedRequests()
 {
-	FScopeLock ScopeLock(&NumRetriedRequestsLock);
-
+	check(IsInGameThread());
 	--NumRetriedRequests;
 	check(NumRetriedRequests >= 0);
 	if (NumRetriedRequests == 0)
diff --git a/Engine/Source/Runtime/Online/HTTP/Private/HttpThread.cpp b/Engine/Source/Runtime/Online/HTTP/Private/HttpThread.cpp
--- a/Engine/Source/Runtime/Online/HTTP/Private/HttpThread.cpp
+++ b/Engine/Source/Runtime/Online/HTTP/Private/HttpThread.cpp
@@ -110,10 +110,13 @@ uint32 FHttpThreadBase::Run()
 
 void FHttpThreadBase::Tick()
 {
-	const double AppTime = FPlatformTime::Seconds();
-	const double ElapsedTime = AppTime - LastTime;
-	LastTime = AppTime;
-	HttpThreadTick(ElapsedTime);
+	// Run HttpThread tasks
+	TFunction<void()> Task = nullptr;
+	while (HttpThreadQueue.Dequeue(Task))
+	{
+		check(Task);
+		Task();
+	}
 }
 
 bool FHttpThreadBase::NeedsSingleThreadTick() const
@@ -137,8 +140,17 @@ void FHttpThreadBase::UpdateConfigs()
 	}
 }
 
+void FHttpThreadBase::AddHttpThreadTask(TFunction<void()>&& Task)
+{
+	if (Task)
+	{
+		HttpThreadQueue.Enqueue(MoveTemp(Task));
+	}
+}
+
 void FHttpThreadBase::HttpThreadTick(float DeltaSeconds)
 {
+	// empty
 }
 
 bool FHttpThreadBase::StartThreadedRequest(IHttpThreadedRequest* Request)
@@ -347,21 +359,6 @@ void FLegacyHttpThread::Tick()
 	}
 }
 
-void FLegacyHttpThread::AddHttpThreadTask(TFunction<void()>&& Task, float InDelay)
-{
-	Ticker.AddTicker(FTickerDelegate::CreateLambda([this, Task=MoveTemp(Task)](float) {
-		Task();
-		return false;
-	}), InDelay);
-}
-
-void FLegacyHttpThread::HttpThreadTick(float DeltaSeconds)
-{
-	FHttpThreadBase::HttpThreadTick(DeltaSeconds);
-
-	Ticker.Tick(DeltaSeconds);
-}
-
 bool FLegacyHttpThread::Init()
 {
 	ExitRequest.Set(false);
diff --git a/Engine/Source/Runtime/Online/HTTP/Private/HttpThread.h b/Engine/Source/Runtime/Online/HTTP/Private/HttpThread.h
--- a/Engine/Source/Runtime/Online/HTTP/Private/HttpThread.h
+++ b/Engine/Source/Runtime/Online/HTTP/Private/HttpThread.h
@@ -10,7 +10,6 @@
 #include "Misc/Timespan.h"
 #include "Containers/Queue.h"
 #include "Containers/SpscQueue.h"
-#include "Containers/Ticker.h"
 
 #include <atomic>
 
@@ -82,10 +81,9 @@ public:
 	/**
 	 * Add task to be ran on the http thread next tick
 	 *
-	 * @param Task The task to be ran
-	 * @param InDelay The delay to wait before running the task
+	 * @param Task The task to be ran next tick
 	 */
-	virtual void AddHttpThreadTask(TFunction<void()>&& Task, float InDelay) = 0;
+	void AddHttpThreadTask(TFunction<void()>&& Task);
 
 protected:
 
@@ -104,6 +102,7 @@ protected:
 	 */
 	virtual void CompleteThreadedRequest(IHttpThreadedRequest* Request);
 
+
 protected:
 	int32 GetRunningThreadedRequestLimit() const;
 
@@ -171,6 +170,9 @@ protected:
 	 * Added to on HTTP thread, processed then cleared on game thread (Single producer, single consumer)
 	 */
 	TSpscQueue<IHttpThreadedRequest*> CompletedThreadedRequests;
+
+	/** Queue of tasks to run on the game thread */
+	TQueue<TFunction<void()>, EQueueMode::Mpsc> HttpThreadQueue;
 };
 
 class FLegacyHttpThread	: public FHttpThreadBase
@@ -199,9 +201,6 @@ protected:
 	virtual void Stop() override;
 	//~ End FRunnable Interface
 
-	virtual void AddHttpThreadTask(TFunction<void()>&& Task, float InDelay) override;
-	virtual void HttpThreadTick(float DeltaSeconds) override;
-
 	/** signal request to stop and exit thread */
 	FThreadSafeCounter ExitRequest;
 
@@ -213,7 +212,4 @@ protected:
 	double HttpThreadIdleFrameTimeInSeconds;
 	/** Time in seconds to sleep minimally when idle, waiting for requests. */
 	double HttpThreadIdleMinimumSleepTimeInSeconds;
-
-	/* Ticker for functions to run in HTTP thread */
-	FTSTicker Ticker;
 };
diff --git a/Engine/Source/Runtime/Online/HTTP/Private/IHttpThreadedRequest.cpp b/Engine/Source/Runtime/Online/HTTP/Private/IHttpThreadedRequest.cpp
--- a/Engine/Source/Runtime/Online/HTTP/Private/IHttpThreadedRequest.cpp
+++ b/Engine/Source/Runtime/Online/HTTP/Private/IHttpThreadedRequest.cpp
@@ -4,20 +4,22 @@
 #include "HttpModule.h"
 #include "HttpManager.h"
 
-void IHttpThreadedRequest::FinishRequestNotInHttpManager()
+bool IHttpThreadedRequest::FinishRequestNotInHttpManager()
 {
 	if (IsInGameThread())
 	{
 		if (DelegateThreadPolicy == EHttpRequestDelegateThreadPolicy::CompleteOnGameThread)
 		{
 			FinishRequest();
+			return false;
 		}
 		else
 		{
 			FHttpModule::Get().GetHttpManager().AddHttpThreadTask([StrongThis = StaticCastSharedRef<IHttpThreadedRequest>(AsShared())]()
-			{
-				StrongThis->FinishRequest();
-			});
+				{
+					StrongThis->FinishRequest();
+				});
+			return true;
 		}
 	}
 	else
@@ -25,13 +27,15 @@ void IHttpThreadedRequest::FinishRequestNotInHttpManager()
 		if (DelegateThreadPolicy == EHttpRequestDelegateThreadPolicy::CompleteOnHttpThread)
 		{
 			FinishRequest();
+			return false;
 		}
 		else
 		{
 			FHttpModule::Get().GetHttpManager().AddGameThreadTask([StrongThis = StaticCastSharedRef<IHttpThreadedRequest>(AsShared())]()
-			{
-				StrongThis->FinishRequest();
-			});
+				{
+					StrongThis->FinishRequest();
+				});
+			return true;
 		}
 	}
 }
diff --git a/Engine/Source/Runtime/Online/HTTP/Private/IHttpThreadedRequest.h b/Engine/Source/Runtime/Online/HTTP/Private/IHttpThreadedRequest.h
--- a/Engine/Source/Runtime/Online/HTTP/Private/IHttpThreadedRequest.h
+++ b/Engine/Source/Runtime/Online/HTTP/Private/IHttpThreadedRequest.h
@@ -19,6 +19,8 @@ public:
 protected:
 	/**
 	 * Finish the request when it's not in http manager
+	 * 
+	 * @return false if the request finished immediate without waiting, otherwise it needs to be finished in another thread instead of current thread
 	 */
-	void FinishRequestNotInHttpManager();
+	bool FinishRequestNotInHttpManager();
 };
diff --git a/Engine/Source/Runtime/Online/HTTP/Public/HttpManager.h b/Engine/Source/Runtime/Online/HTTP/Public/HttpManager.h
--- a/Engine/Source/Runtime/Online/HTTP/Public/HttpManager.h
+++ b/Engine/Source/Runtime/Online/HTTP/Public/HttpManager.h
@@ -232,10 +232,9 @@ public:
 	/**
 	 * Add task to be ran on the http thread next tick
 	 *
-	 * @param Task The task to be ran
-	 * @param InDelay The delay to wait before running the task
+	 * @param Task The task to be ran next tick
 	 */
-	HTTP_API void AddHttpThreadTask(TFunction<void()>&& Task, float InDelay = 0.0f);
+	HTTP_API void AddHttpThreadTask(TFunction<void()>&& Task);
 
 	/**
 	 * Set url request filter through code, instead of setting it through config.
@@ -316,9 +315,4 @@ PACKAGE_SCOPE:
 	 * Called explicitly by non-threaded requests
 	 */
 	HTTP_API void BroadcastHttpRequestCompleted(const FHttpRequestRef& Request);
-
-	/**
-	 * Access http thread of http manager for internal usage
-	 */
-	HTTP_API FHttpThreadBase* GetThread();
 };
diff --git a/Engine/Source/Runtime/Online/HTTP/Public/HttpRetrySystem.h b/Engine/Source/Runtime/Online/HTTP/Public/HttpRetrySystem.h
--- a/Engine/Source/Runtime/Online/HTTP/Public/HttpRetrySystem.h
+++ b/Engine/Source/Runtime/Online/HTTP/Public/HttpRetrySystem.h
@@ -98,7 +98,7 @@ namespace FHttpRetrySystem
 		HTTP_API virtual void CancelRequest() override;
 
 		// FRequest
-		EStatus::Type GetRetryStatus() const { return RetryStatus; }
+		EStatus::Type GetRetryStatus() const { return Status; }
 
 	protected:
 		friend class FManager;
@@ -114,15 +114,13 @@ namespace FHttpRetrySystem
 		);
 
 		void HttpOnRequestProgress(FHttpRequestPtr InHttpRequest, uint64 BytesSent, uint64 BytesRcv);
-		void HttpOnProcessRequestComplete(FHttpRequestPtr HttpRequest, FHttpResponsePtr HttpResponse, bool bSucceeded);
-		void HttpOnHeaderReceived(FHttpRequestPtr Request, const FString& HeaderName, const FString& NewHeaderValue);
 
 		/** Update our HTTP request's URL's domain from our RetryDomains */
 		void SetUrlFromRetryDomains();
 		/** Move to the next retry domain from our RetryDomains */
 		void MoveToNextRetryDomain();
 
-		EStatus::Type                        RetryStatus;
+		EStatus::Type                        Status;
 
 		FRetryLimitCountSetting              RetryLimitCountOverride;
 		FRetryTimeoutRelativeSecondsSetting  RetryTimeoutRelativeSecondsOverride;
@@ -170,9 +168,7 @@ public:
 	 *
 	 * @return                true if there are no failures or retries
 	 */
-	UE_DEPRECATED(5.4, "HttpRetrySystem::Update has been deprecated, all logic will be processed in http thread instead")
 	HTTP_API bool Update(uint32* FileCount = NULL, uint32* FailingCount = NULL, uint32* FailedCount = NULL, uint32* CompletedCount = NULL);
-
 	void SetRandomFailureRate(float Value) { RandomFailureRate = FRandomFailureRateSetting(Value); }
 	void SetDefaultRetryLimit(uint32 Value) { RetryLimitCountDefault = FRetryLimitCountSetting(Value); }
 
@@ -212,8 +208,6 @@ protected:
 
 		/** Number of requests that are in a retried state.  When this is non-zero, verbosity will be adjusted. */
 		int32 NumRetriedRequests = 0;
-		/** DecrementRetriedRequests can be called from game thread or http thread depends on the http request thread policy, make sure it's thread-safe */
-		FCriticalSection NumRetriedRequestsLock;
 		/** Verbosity to restore to when there are no requests being retried */
 		ELogVerbosity::Type OriginalVerbosity = ELogVerbosity::Error;
 		/** Config driven target verbosity to set to when requests are being retried.  NoLogging means the verbosity will not be modified. */
@@ -239,13 +233,6 @@ protected:
 	 */
 	void RetryHttpRequest(FHttpRetryRequestEntry& RequestEntry);
 
-	/**
-	 * Retry an HTTP request with delay
-	 * @param RequestToRetry request to retry
-	 * @param InDelay the delay to wait before retrying
-	 */
-	void RetryHttpRequestWithDelay(const TSharedRef<FRequest>& Request, float InDelay);
-
 	// @return number of seconds to lockout for
 	float GetLockoutPeriodSeconds(const FHttpRetryRequestEntry& HttpRetryRequestEntry);
 
