Revert the changes to limit schemes to http and https. We already have support to config the allowed schemes through FURLRequestFilter.

diff --git a/Engine/Source/Programs/WebTests/Private/TestHttp.cpp b/Engine/Source/Programs/WebTests/Private/TestHttp.cpp
--- a/Engine/Source/Programs/WebTests/Private/TestHttp.cpp
+++ b/Engine/Source/Programs/WebTests/Private/TestHttp.cpp
@@ -572,20 +572,31 @@ TEST_CASE_METHOD(FWaitUntilCompleteHttpFixture, "Can upload big file exceeds 32
 	HttpRequest->ProcessRequest();
 }
 
-TEST_CASE_METHOD(FWaitUntilCompleteHttpFixture, "Streaming http upload from file by PUT can work well", HTTP_TAG)
+namespace UE
+{
+namespace TestHttp
 {
-	FString Filename = FString(FPlatformProcess::UserSettingsDir()) / TEXT("TestStreamUpload.dat");
 
-	FArchive* RawFile = IFileManager::Get().CreateFileWriter(*Filename);
+void WriteTestFile(const FString& TestFileName, uint64 TestFileSize)
+{
+	FArchive* RawFile = IFileManager::Get().CreateFileWriter(*TestFileName);
 	CHECK(RawFile != nullptr);
 	TSharedRef<FArchive> FileToWrite = MakeShareable(RawFile);
-	const uint64 FileSize = 5*1024*1024; // 5MB
-	char* FileData = (char*)FMemory::Malloc(FileSize);
-	FMemory::Memset(FileData, 'd', FileSize);
-	FileToWrite->Serialize(FileData, FileSize);
+	char* FileData = (char*)FMemory::Malloc(TestFileSize);
+	FMemory::Memset(FileData, 'd', TestFileSize);
+	FileToWrite->Serialize(FileData, TestFileSize);
 	FileToWrite->FlushCache();
 	FileToWrite->Close();
 	FMemory::Free(FileData);
+}
+
+}
+}
+
+TEST_CASE_METHOD(FWaitUntilCompleteHttpFixture, "Streaming http upload from file by PUT can work well", HTTP_TAG)
+{
+	FString Filename = FString(FPlatformProcess::UserSettingsDir()) / TEXT("TestStreamUpload.dat");
+	UE::TestHttp::WriteTestFile(Filename, 5*1024*1024/*5MB*/);
 
 	TSharedRef<IHttpRequest> HttpRequest = CreateRequest();
 	HttpRequest->SetURL(FString::Format(TEXT("{0}/streaming_upload_put"), { *UrlHttpTests() }));
@@ -1037,6 +1048,36 @@ TEST_CASE_METHOD(FThreadedBatchRequestsFixture, "Retry manager and http manager
 	BlockUntilFlushed();
 }
 
+#if (PLATFORM_WINDOWS && !WITH_CURL_XCURL) || PLATFORM_MAC || PLATFORM_UNIX
+
+TEST_CASE_METHOD(FWaitUntilCompleteHttpFixture, "Scheme besides http and https can work if allowed by settings", HTTP_TAG)
+{
+	bool bShouldSucceed = false;
+	SECTION("when allowed")
+	{
+		bShouldSucceed = true;
+	}
+	SECTION("when not allowed")
+	{
+		DisableWarningsInThisTest();
+		// Pre check will fail when scheme is not listed
+		UE::TestHttp::SetupURLRequestFilter(HttpModule);
+	}
+
+	FString Filename = FString(FPlatformProcess::UserSettingsDir()) / TEXT("TestProtocolAllowed.dat");
+	UE::TestHttp::WriteTestFile(Filename, 10/*Bytes*/);
+
+	TSharedRef<IHttpRequest, ESPMode::ThreadSafe> HttpRequest = HttpModule->CreateRequest();
+	HttpRequest->SetURL(FString(TEXT("file://")) + Filename.Replace(TEXT(" "), TEXT("%20")));
+	HttpRequest->SetVerb(TEXT("GET"));
+	HttpRequest->OnProcessRequestComplete().BindLambda([Filename, bShouldSucceed](FHttpRequestPtr HttpRequest, FHttpResponsePtr HttpResponse, bool bSucceeded) {
+		CHECK(bSucceeded == bShouldSucceed);
+		IFileManager::Get().Delete(*Filename);
+	});
+	HttpRequest->ProcessRequest();
+}
+
+#endif
 
 // TODO: Add cancel test, with multiple cancel calls
 
diff --git a/Engine/Source/Runtime/Online/HTTP/HTTP.Build.cs b/Engine/Source/Runtime/Online/HTTP/HTTP.Build.cs
--- a/Engine/Source/Runtime/Online/HTTP/HTTP.Build.cs
+++ b/Engine/Source/Runtime/Online/HTTP/HTTP.Build.cs
@@ -82,7 +82,7 @@ public class HTTP : ModuleRules
 		}
 
 		PrivateDefinitions.Add("WITH_CURL_LIBCURL =" + (bPlatformSupportsLibCurl ? "1" : "0"));
-		PrivateDefinitions.Add("WITH_CURL_XCURL=" + (bPlatformSupportsXCurl ? "1" : "0"));
+		PublicDefinitions.Add("WITH_CURL_XCURL=" + (bPlatformSupportsXCurl ? "1" : "0"));
 		PrivateDefinitions.Add("WITH_CURL_MULTIPOLL=" + (bPlatformSupportsCurlMultiPoll ? "1" : "0"));
 		PrivateDefinitions.Add("WITH_CURL_MULTIWAIT=" + (bPlatformSupportsCurlMultiWait ? "1" : "0"));
 		PrivateDefinitions.Add("WITH_CURL_MULTISOCKET=" + (bPlatformSupportsCurlMultiSocket ? "1" : "0"));
diff --git a/Engine/Source/Runtime/Online/HTTP/Private/GenericPlatform/HttpRequestCommon.cpp b/Engine/Source/Runtime/Online/HTTP/Private/GenericPlatform/HttpRequestCommon.cpp
--- a/Engine/Source/Runtime/Online/HTTP/Private/GenericPlatform/HttpRequestCommon.cpp
+++ b/Engine/Source/Runtime/Online/HTTP/Private/GenericPlatform/HttpRequestCommon.cpp
@@ -42,14 +42,6 @@ bool FHttpRequestCommon::PreCheck() const
 		return false;
 	}
 
-	// Make sure the URL is parsed correctly with a valid HTTP scheme
-	const FString URL = GetURL();
-	if (!URL.StartsWith(TEXT("http://")) && !URL.StartsWith(TEXT("https://")))
-	{
-		UE_LOG(LogHttp, Warning, TEXT("ProcessRequest failed. URL '%s' is not a valid HTTP request."), *GetURL());
-		return false;
-	}
-
 	if (!FHttpModule::Get().GetHttpManager().IsDomainAllowed(GetURL()))
 	{
 		UE_LOG(LogHttp, Warning, TEXT("ProcessRequest failed. URL '%s' is not using an allowed domain."), *GetURL());
