Make all http impls respect the connect timeout set in http module.

#tests Test covered in WebTests project

diff --git a/Engine/Source/Programs/WebTests/Private/TestHttp.cpp b/Engine/Source/Programs/WebTests/Private/TestHttp.cpp
--- a/Engine/Source/Programs/WebTests/Private/TestHttp.cpp
+++ b/Engine/Source/Programs/WebTests/Private/TestHttp.cpp
@@ -3,6 +3,7 @@
 #include "HAL/FileManager.h"
 #include "HAL/IConsoleManager.h"
 #include "HAL/PlatformProcess.h"
+#include "HAL/PlatformTime.h"
 #include "HAL/Runnable.h"
 #include "HAL/RunnableThread.h"
 #include "HttpManager.h"
@@ -32,6 +33,12 @@
 extern TAutoConsoleVariable<bool> CVarHttpInsecureProtocolEnabled;
 extern TAutoConsoleVariable<bool> CVarHttpRetrySystemNonGameThreadSupportEnabled;
 
+class FMockHttpModule : public FHttpModule
+{
+public:
+	using FHttpModule::HttpConnectionTimeout;
+};
+
 class FHttpModuleTestFixture
 {
 public:
@@ -46,7 +53,7 @@ public:
 
 		bRetryEnabled &= CVarHttpRetrySystemNonGameThreadSupportEnabled.GetValueOnAnyThread();
 
-		HttpModule = new FHttpModule();
+		HttpModule = new FMockHttpModule();
 		IModuleInterface* Module = HttpModule;
 		Module->StartupModule();
 
@@ -77,7 +84,7 @@ public:
 		LogHttp.SetVerbosity(ELogVerbosity::Error);
 	}
 
-	const FString UrlWithInvalidPortToTestConnectTimeout() const { return FString::Format(TEXT("http://{0}:{1}"), { *WebServerIp, 8765 }); }
+	const FString UrlWithInvalidPortToTestConnectTimeout() const { return TEXT("http://10.255.255.1:8765"); } // non-routable IP address with a random port
 	const FString UrlBase() const { return FString::Format(TEXT("http://{0}:{1}"), { *WebServerIp, WebServerHttpPort }); }
 	const FString UrlHttpTests() const { return FString::Format(TEXT("{0}/webtests/httptests"), { *UrlBase() }); }
 	const FString UrlToTestMethods() const { return FString::Format(TEXT("{0}/methods"), { *UrlHttpTests() }); }
@@ -85,7 +92,7 @@ public:
 
 	FString WebServerIp;
 	uint32 WebServerHttpPort;
-	FHttpModule* HttpModule;
+	FMockHttpModule* HttpModule;
 
 	bool bRunHeavyTests;
 	bool bRetryEnabled;
@@ -145,7 +152,6 @@ public:
 		if (bRetryEnabled)
 		{
 			HttpRetryManager = MakeShared<FMockRetryManager>(FHttpRetrySystem::FRetryLimitCountSetting(RetryLimitCount), FHttpRetrySystem::FRetryTimeoutRelativeSecondsSetting(/*RetryTimeoutRelativeSeconds*/));
-			HttpRetryManager->RetryTimeoutRelativeSecondsDefault = 2.0f; // Value is set so that retry system will use it at high level, will no longer wait for ever for HTTP code
 		}
 	}
 
@@ -240,26 +246,33 @@ TEST_CASE_METHOD(FWaitUntilCompleteHttpFixture, "Get large response content with
 	HttpRequest->ProcessRequest();
 }
 
-// TODO: Enable this after finding a more reliable way of simulate timeout
-//TEST_CASE_METHOD(FWaitUntilCompleteHttpFixture, "Http request connect timeout", HTTP_TAG)
-//{
-//	TSharedRef<IHttpRequest> HttpRequest = CreateRequest();
-//	HttpRequest->SetURL(UrlWithInvalidPortToTestConnectTimeout());
-//	HttpRequest->SetVerb(TEXT("GET"));
-//	HttpRequest->SetTimeout(7);
-//	FDateTime StartTime = FDateTime::Now();
-//	HttpRequest->OnProcessRequestComplete().BindLambda([StartTime](FHttpRequestPtr HttpRequest, FHttpResponsePtr HttpResponse, bool bSucceeded) {
-//		CHECK(!bSucceeded);
-//		CHECK(HttpResponse == nullptr);
-//		// TODO: For now curl impl is using customized timeout instead of relying on native http timeout, 
-//		// which doesn't get CURLE_COULDNT_CONNECT. Enable this after switching to native http timeout
-//		//CHECK(HttpRequest->GetStatus() == EHttpRequestStatus::Failed_ConnectionError);
-//		FTimespan Timespan = FDateTime::Now() - StartTime;
-//		float DurationInSeconds = Timespan.GetTotalSeconds();
-//		CHECK(FMath::IsNearlyEqual(DurationInSeconds, 7, HTTP_TIME_DIFF_TOLERANCE));
-//	});
-//	HttpRequest->ProcessRequest();
-//}
+TEST_CASE_METHOD(FWaitUntilCompleteHttpFixture, "Http request connect timeout", HTTP_TAG)
+{
+	DisableWarningsInThisTest();
+
+	HttpModule->HttpConnectionTimeout = 15.0f;
+
+	TSharedRef<IHttpRequest> HttpRequest = CreateRequest();
+
+	HttpRequest->SetURL(UrlWithInvalidPortToTestConnectTimeout());
+	HttpRequest->SetVerb(TEXT("GET"));
+
+	const double StartTime = FPlatformTime::Seconds();
+
+	HttpRequest->OnProcessRequestComplete().BindLambda([StartTime](FHttpRequestPtr HttpRequest, FHttpResponsePtr HttpResponse, bool bSucceeded) {
+		CHECK(!bSucceeded);
+		CHECK(HttpResponse == nullptr);
+		CHECK(HttpRequest->GetStatus() == EHttpRequestStatus::Failed);
+		CHECK(HttpRequest->GetFailureReason() == EHttpFailureReason::ConnectionError);
+		const double DurationInSeconds  = FPlatformTime::Seconds() - StartTime;
+		double HttpTimeDiffTolerance = 1.0;
+#if WITH_CURL_XCURL
+		HttpTimeDiffTolerance += 3.0; // It seems xCurl takes up to 3 more seconds for connect timeout
+#endif
+		CHECK(FMath::IsNearlyEqual(DurationInSeconds, 15.0, HttpTimeDiffTolerance));
+	});
+	HttpRequest->ProcessRequest();
+}
 
 TEST_CASE_METHOD(FWaitUntilCompleteHttpFixture, "Streaming http download", HTTP_TAG)
 {
@@ -990,17 +1003,16 @@ TEST_CASE_METHOD(FWaitUntilCompleteHttpFixture, "Cancel http request connect bef
 {
 	DisableWarningsInThisTest();
 
-	TSharedRef<IHttpRequest, ESPMode::ThreadSafe> HttpRequest = CreateRequest();
+	TSharedRef<IHttpRequest> HttpRequest = CreateRequest();
 	HttpRequest->SetURL(UrlWithInvalidPortToTestConnectTimeout());
 	HttpRequest->SetVerb(TEXT("GET"));
 	HttpRequest->SetTimeout(7);
-	FDateTime StartTime = FDateTime::Now();
+	const double StartTime = FPlatformTime::Seconds();
 	HttpRequest->OnProcessRequestComplete().BindLambda([StartTime](FHttpRequestPtr HttpRequest, FHttpResponsePtr HttpResponse, bool bSucceeded) {
 		CHECK(!bSucceeded);
-
-		FTimespan Timespan = FDateTime::Now() - StartTime;
-		float DurationInSeconds = Timespan.GetTotalSeconds();
-		CHECK(DurationInSeconds < 2);
+		const double DurationInSeconds = FPlatformTime::Seconds() - StartTime;
+		CHECK(DurationInSeconds < 2.0);
+		CHECK(HttpRequest->GetFailureReason() == EHttpFailureReason::Cancelled);
 	});
 	HttpRequest->ProcessRequest();
 	FPlatformProcess::Sleep(0.5);
@@ -1067,7 +1079,7 @@ TEST_CASE_METHOD(FWaitUntilCompleteHttpFixture, "Scheme besides http and https c
 	FString Filename = FString(FPlatformProcess::UserSettingsDir()) / TEXT("TestProtocolAllowed.dat");
 	UE::TestHttp::WriteTestFile(Filename, 10/*Bytes*/);
 
-	TSharedRef<IHttpRequest, ESPMode::ThreadSafe> HttpRequest = HttpModule->CreateRequest();
+	TSharedRef<IHttpRequest> HttpRequest = HttpModule->CreateRequest();
 	HttpRequest->SetURL(FString(TEXT("file://")) + Filename.Replace(TEXT(" "), TEXT("%20")));
 	HttpRequest->SetVerb(TEXT("GET"));
 	HttpRequest->OnProcessRequestComplete().BindLambda([Filename, bShouldSucceed](FHttpRequestPtr HttpRequest, FHttpResponsePtr HttpResponse, bool bSucceeded) {
diff --git a/Engine/Source/Runtime/Online/HTTP/Private/Apple/AppleHttp.cpp b/Engine/Source/Runtime/Online/HTTP/Private/Apple/AppleHttp.cpp
--- a/Engine/Source/Runtime/Online/HTTP/Private/Apple/AppleHttp.cpp
+++ b/Engine/Source/Runtime/Online/HTTP/Private/Apple/AppleHttp.cpp
@@ -377,7 +377,9 @@ FAppleHttpRequest::FAppleHttpRequest(NSURLSession* InSession)
 {
 	UE_LOG(LogHttp, Verbose, TEXT("FAppleHttpRequest::FAppleHttpRequest()"));
 	Request = [[NSMutableURLRequest alloc] init];
-	Request.timeoutInterval = FHttpModule::Get().GetHttpTimeout();
+	float HttpConnectionTimeout = FHttpModule::Get().GetHttpConnectionTimeout();
+	check(HttpConnectionTimeout > 0.0f);
+	Request.timeoutInterval = HttpConnectionTimeout;
 
 	// Disable cache to mimic WinInet behavior
 	Request.cachePolicy = NSURLRequestReloadIgnoringLocalCacheData;
@@ -632,21 +634,6 @@ void FAppleHttpRequest::SetVerb(const FString& Verb)
 	Request.HTTPMethod = Verb.GetNSString();
 }
 
-void FAppleHttpRequest::SetTimeout(float InTimeoutSecs)
-{
-	Request.timeoutInterval = InTimeoutSecs;
-}
-
-void FAppleHttpRequest::ClearTimeout()
-{
-	Request.timeoutInterval = FHttpModule::Get().GetHttpTimeout();
-}
-
-TOptional<float> FAppleHttpRequest::GetTimeout() const
-{
-	return TOptional<float>(Request.timeoutInterval);
-}
-
 bool FAppleHttpRequest::ProcessRequest()
 {
 	SCOPED_AUTORELEASE_POOL;
diff --git a/Engine/Source/Runtime/Online/HTTP/Private/Apple/AppleHttp.h b/Engine/Source/Runtime/Online/HTTP/Private/Apple/AppleHttp.h
--- a/Engine/Source/Runtime/Online/HTTP/Private/Apple/AppleHttp.h
+++ b/Engine/Source/Runtime/Online/HTTP/Private/Apple/AppleHttp.h
@@ -45,9 +45,6 @@ public:
 	virtual bool SetResponseBodyReceiveStream(TSharedRef<FArchive> Stream) override;
 	virtual void SetHeader(const FString& HeaderName, const FString& HeaderValue) override;
 	virtual void AppendToHeader(const FString& HeaderName, const FString& AdditionalHeaderValue) override;
-	virtual void SetTimeout(float InTimeoutSecs) override;
-	virtual void ClearTimeout() override;
-	virtual TOptional<float> GetTimeout() const override;
 	virtual bool ProcessRequest() override;
 	virtual void CancelRequest() override;
 	virtual const FHttpResponsePtr GetResponse() const override;
diff --git a/Engine/Source/Runtime/Online/HTTP/Private/Curl/CurlHttp.cpp b/Engine/Source/Runtime/Online/HTTP/Private/Curl/CurlHttp.cpp
--- a/Engine/Source/Runtime/Online/HTTP/Private/Curl/CurlHttp.cpp
+++ b/Engine/Source/Runtime/Online/HTTP/Private/Curl/CurlHttp.cpp
@@ -941,10 +941,8 @@ bool FCurlHttpRequest::SetupRequestHttpThread()
 
 		// Set connection timeout in seconds
 		int32 HttpConnectionTimeout = FHttpModule::Get().GetHttpConnectionTimeout();
-		if (HttpConnectionTimeout >= 0)
-		{
-			curl_easy_setopt(EasyHandle, CURLOPT_CONNECTTIMEOUT, HttpConnectionTimeout);
-		}
+		check(HttpConnectionTimeout > 0);
+		curl_easy_setopt(EasyHandle, CURLOPT_CONNECTTIMEOUT, HttpConnectionTimeout);
 
 		if (FCurlHttpManager::CurlRequestOptions.bAllowSeekFunction && bIsRequestPayloadSeekable)
 		{
@@ -1302,9 +1300,13 @@ void FCurlHttpRequest::FinishRequest()
 			switch (CurlCompletionResult)
 			{
 			case CURLE_COULDNT_CONNECT:
+			case CURLE_OPERATION_TIMEDOUT:
 			case CURLE_COULDNT_RESOLVE_PROXY:
 			case CURLE_COULDNT_RESOLVE_HOST:
 			case CURLE_SSL_CONNECT_ERROR:
+#if WITH_CURL_XCURL
+			case CURLE_SEND_ERROR:
+#endif
 				// report these as connection errors (safe to retry)
 				SetFailureReason(EHttpFailureReason::ConnectionError);
 				break;
diff --git a/Engine/Source/Runtime/Online/HTTP/Private/HttpModule.cpp b/Engine/Source/Runtime/Online/HTTP/Private/HttpModule.cpp
--- a/Engine/Source/Runtime/Online/HTTP/Private/HttpModule.cpp
+++ b/Engine/Source/Runtime/Online/HTTP/Private/HttpModule.cpp
@@ -78,7 +78,7 @@ void FHttpModule::StartupModule()
 
 	MaxReadBufferSize = 256 * 1024;
 	HttpTimeout = 300.0f;
-	HttpConnectionTimeout = -1;
+	HttpConnectionTimeout = 60.0f;
 	HttpReceiveTimeout = HttpConnectionTimeout;
 	HttpSendTimeout = HttpConnectionTimeout;
 	HttpMaxConnectionsPerServer = 16;
diff --git a/Engine/Source/Runtime/Online/HTTP/Public/HttpModule.h b/Engine/Source/Runtime/Online/HTTP/Public/HttpModule.h
--- a/Engine/Source/Runtime/Online/HTTP/Public/HttpModule.h
+++ b/Engine/Source/Runtime/Online/HTTP/Public/HttpModule.h
@@ -315,6 +315,10 @@ public:
 		return AllowedDomains;
 	}
 
+protected:
+	/** timeout in seconds to establish the connection */
+	float HttpConnectionTimeout;
+
 private:
 
 	// IModuleInterface
@@ -351,11 +355,9 @@ private:
 	FHttpManager* HttpManager = nullptr;
 	/** timeout in seconds for the entire http request to complete. 0 is no timeout */
 	float HttpTimeout;
-	/** timeout in seconds to establish the connection. -1 for system defaults, 0 is no timeout */
-	float HttpConnectionTimeout;
-	/** timeout in seconds to receive a response on the connection. -1 for system defaults */
+	/** timeout in seconds to receive a response on the connection */
 	float HttpReceiveTimeout;
-	/** timeout in seconds to send a request on the connection. -1 for system defaults */
+	/** timeout in seconds to send a request on the connection */
 	float HttpSendTimeout;
 	/** total time to delay the request */
 	float HttpDelayTime;
