Add test case to validate the http header receive callback.

#tests Tested through WebTests preflight

diff --git a/Engine/Source/Programs/WebTests/Private/TestHttp.cpp b/Engine/Source/Programs/WebTests/Private/TestHttp.cpp
--- a/Engine/Source/Programs/WebTests/Private/TestHttp.cpp
+++ b/Engine/Source/Programs/WebTests/Private/TestHttp.cpp
@@ -910,6 +910,52 @@ TEST_CASE_METHOD(FWaitUntilCompleteHttpFixture, "Http request pre check will fai
 	HttpRequest->ProcessRequest();
 }
 
+class FValidateHeaderReceiveOrderFixture : public FWaitUntilCompleteHttpFixture
+{
+public:
+	~FValidateHeaderReceiveOrderFixture()
+	{
+		WaitUntilAllHttpRequestsComplete();
+	}
+
+	std::atomic<bool> bHeaderReceived = false;
+	std::atomic<bool> bCompleteCallbackTriggered = false;
+};
+
+TEST_CASE_METHOD(FValidateHeaderReceiveOrderFixture, "Http request header received callback will be called by thread policy", HTTP_TAG)
+{
+	TSharedRef<IHttpRequest> HttpRequest = CreateRequest();
+	HttpRequest->SetURL(UrlStreamDownload(2/*Chunks*/, 1024/*ChunkSize*/));
+	HttpRequest->SetVerb(TEXT("GET"));
+
+	SECTION("in http thread")
+	{
+		HttpRequest->SetDelegateThreadPolicy(EHttpRequestDelegateThreadPolicy::CompleteOnHttpThread);
+		HttpRequest->OnHeaderReceived().BindLambda([this](FHttpRequestPtr Request, const FString& HeaderName, const FString& HeaderValue) {
+			CHECK(!bCompleteCallbackTriggered);
+			CHECK(!IsInGameThread());
+			bHeaderReceived = true;
+		});
+	}
+	SECTION("in game thread")
+	{
+		HttpRequest->SetDelegateThreadPolicy(EHttpRequestDelegateThreadPolicy::CompleteOnGameThread);
+		HttpRequest->OnHeaderReceived().BindLambda([this](FHttpRequestPtr Request, const FString& HeaderName, const FString& HeaderValue) {
+			CHECK(!bCompleteCallbackTriggered);
+			CHECK(IsInGameThread());
+			bHeaderReceived = true;
+		});
+	}
+
+	HttpRequest->OnProcessRequestComplete().BindLambda([this](FHttpRequestPtr /*HttpRequest*/, FHttpResponsePtr /*HttpResponse */, bool bSucceeded) {
+		CHECK(bHeaderReceived);
+		bCompleteCallbackTriggered = true;
+		CHECK(bSucceeded);
+	});
+
+	HttpRequest->ProcessRequest();
+}
+
 namespace UE
 {
 namespace TestHttp
@@ -1169,4 +1215,3 @@ TEST_CASE_METHOD(FLocalHttpServerFixture, "Local http server can serve large fil
 
 // TODO: Add cancel test, with multiple cancel calls
 
-// TODO: Add test case to validate header received callback can be received in http/game thread, and can be received before the request complete
