Add test code to cover the case that download progress get cleaned when retry.

#tests Tested through WebTests in local

diff --git a/Engine/Source/Programs/WebTests/Private/TestHttp.cpp b/Engine/Source/Programs/WebTests/Private/TestHttp.cpp
--- a/Engine/Source/Programs/WebTests/Private/TestHttp.cpp
+++ b/Engine/Source/Programs/WebTests/Private/TestHttp.cpp
@@ -680,6 +680,7 @@ TEST_CASE_METHOD(FWaitUntilQuitFromTestFixture, "Http request can be reused", HT
 		CHECK(bSucceeded);
 		CHECK(HttpResponse->GetResponseCode() == 200);
 
+		// Using a different URL
 		uint32 Chunks = 3;
 		uint32 ChunkSize = 1024;
 		HttpRequest->SetURL(UrlStreamDownload(Chunks, ChunkSize));
@@ -689,7 +690,16 @@ TEST_CASE_METHOD(FWaitUntilQuitFromTestFixture, "Http request can be reused", HT
 			REQUIRE(HttpResponse != nullptr);
 			CHECK(HttpResponse->GetResponseCode() == 200);
 			CHECK(HttpResponse->GetContentLength() == Chunks * ChunkSize);
-			bQuitRequested = true;
+
+			// Simulate retry with same URL info
+			HttpRequest->OnProcessRequestComplete().BindLambda([this, Chunks, ChunkSize](FHttpRequestPtr HttpRequest, FHttpResponsePtr HttpResponse, bool bSucceeded) {
+				CHECK(bSucceeded);
+				REQUIRE(HttpResponse != nullptr);
+				CHECK(HttpResponse->GetResponseCode() == 200);
+				CHECK(HttpResponse->GetContentLength() == Chunks * ChunkSize);
+				bQuitRequested = true;
+			});
+			HttpRequest->ProcessRequest();
 		});
 		HttpRequest->ProcessRequest();
 	});
