Add mock_latency method in web tests server, prepare for timeout tests;
Add chunk latency param for streaming api in web tests server;
Make sure streaming api in web tests server can actually return response data by chunks in django asgi server.

#tests Tried in WebTests project

diff --git a/Engine/Source/Programs/WebTests/Private/TestHttp.cpp b/Engine/Source/Programs/WebTests/Private/TestHttp.cpp
--- a/Engine/Source/Programs/WebTests/Private/TestHttp.cpp
+++ b/Engine/Source/Programs/WebTests/Private/TestHttp.cpp
@@ -92,7 +92,9 @@ public:
 	const FString UrlBase() const { return FString::Format(TEXT("http://{0}:{1}"), { *WebServerIp, WebServerHttpPort }); }
 	const FString UrlHttpTests() const { return FString::Format(TEXT("{0}/webtests/httptests"), { *UrlBase() }); }
 	const FString UrlToTestMethods() const { return FString::Format(TEXT("{0}/methods"), { *UrlHttpTests() }); }
-	const FString UrlStreamDownload(uint32 Chunks, uint32 ChunkSize) { return FString::Format(TEXT("{0}/streaming_download/{1}/{2}/"), { *UrlHttpTests(), Chunks, ChunkSize }); }
+	const FString UrlStreamDownload(uint32 Chunks, uint32 ChunkSize, uint32 ChunkLatency=0) { return FString::Format(TEXT("{0}/streaming_download/{1}/{2}/{3}/"), { *UrlHttpTests(), Chunks, ChunkSize, ChunkLatency }); }
+	const FString UrlStreamUpload() { return FString::Format(TEXT("{0}/streaming_upload_put"), { *UrlHttpTests() }); }
+	const FString UrlMockLatency(uint32 Latency) const { return FString::Format(TEXT("{0}/mock_latency/{1}/"), { *UrlHttpTests(), Latency }); }
 
 	FString WebServerIp;
 	uint32 WebServerHttpPort;
@@ -581,7 +583,7 @@ TEST_CASE_METHOD(FWaitUntilCompleteHttpFixture, "Can upload big file exceeds 32
 	TSharedRef<FTestHttpUploadStream> Stream = MakeShared<FTestHttpUploadStream>(TotalSize);
 
 	TSharedRef<IHttpRequest> HttpRequest = CreateRequest();
-	HttpRequest->SetURL(FString::Format(TEXT("{0}/streaming_upload_put"), { *UrlHttpTests() }));
+	HttpRequest->SetURL(UrlStreamUpload());
 	HttpRequest->SetVerb(TEXT("PUT"));
 	HttpRequest->SetContentFromStream(Stream);
 	HttpRequest->SetHeader(TEXT("Content-Disposition"), TEXT("attachment;filename=TestStreamUpload.dat"));
@@ -621,7 +623,7 @@ TEST_CASE_METHOD(FWaitUntilCompleteHttpFixture, "Streaming http upload from file
 	UE::TestHttp::WriteTestFile(Filename, 5*1024*1024/*5MB*/);
 
 	TSharedRef<IHttpRequest> HttpRequest = CreateRequest();
-	HttpRequest->SetURL(FString::Format(TEXT("{0}/streaming_upload_put"), { *UrlHttpTests() }));
+	HttpRequest->SetURL(UrlStreamUpload());
 	HttpRequest->SetVerb(TEXT("PUT"));
 	HttpRequest->SetHeader(TEXT("Content-Disposition"), TEXT("attachment;filename=TestStreamUpload.dat"));
 	HttpRequest->SetContentAsStreamedFile(Filename);
diff --git a/Engine/Source/Programs/WebTestsServer/httptests/urls.py b/Engine/Source/Programs/WebTestsServer/httptests/urls.py
--- a/Engine/Source/Programs/WebTestsServer/httptests/urls.py
+++ b/Engine/Source/Programs/WebTestsServer/httptests/urls.py
@@ -6,10 +6,10 @@ urlpatterns = [
     path("httptests/methods", views.methods, name="methods"),
     path("httptests/query_with_params/", views.query_with_params, name="query_with_params"),
     path("httptests/get_large_response_without_chunks/<int:bytes_number>/", views.get_large_response_without_chunks, name="get_large_response_without_chunks"),
-    path("httptests/nonstreaming_receivetimeout/<int:wait_time>/", views.nonstreaming_receivetimeout, name="nonstreaming_receivetimeout"),
-    path("httptests/streaming_download/<int:chunks>/<int:chunk_size>/", views.streaming_download, name="streaming_download"),
+    path("httptests/streaming_download/<int:chunks>/<int:chunk_size>/<int:chunk_latency>/", views.streaming_download, name="streaming_download"),
     path("httptests/streaming_upload_post", views.streaming_upload_post, name="streaming_upload_post"),
     path("httptests/streaming_upload_put", views.streaming_upload_put, name="streaming_upload_put"),
     path("httptests/redirect_from", views.redirect_from, name="redirect_from"),
     path("httptests/redirect_to", views.redirect_to, name="redirect_to"),
+    path("httptests/mock_latency/<int:latency>/", views.mock_latency, name="mock_latency"),
 ]
diff --git a/Engine/Source/Programs/WebTestsServer/httptests/views.py b/Engine/Source/Programs/WebTestsServer/httptests/views.py
--- a/Engine/Source/Programs/WebTestsServer/httptests/views.py
+++ b/Engine/Source/Programs/WebTestsServer/httptests/views.py
@@ -1,3 +1,5 @@
+import asyncio
+from datetime import datetime
 from django.shortcuts import render, redirect
 from django.http.response import JsonResponse
 from django.http import JsonResponse
@@ -25,15 +27,14 @@ def get_large_response_without_chunks(request, bytes_number):
     return JsonResponse({'data' : data})
 
 @api_view(['GET'])
-def nonstreaming_receivetimeout(request, wait_time):
-    time.sleep(wait_time)
-    return JsonResponse({})
-
-@api_view(['GET'])
-def streaming_download(request, chunks, chunk_size):
-    def data_chunk_generator():
+def streaming_download(request, chunks, chunk_size, chunk_latency):
+    async def data_chunk_generator(): # using 'async' for StreamingHttpResponse in asgi
         for x in range(chunks):
+            if chunk_latency > 0:
+                logger.debug("[%s]sleeping %d seconds", datetime.now().strftime("%H:%M:%S:%f"), chunk_latency)
+                await asyncio.sleep(chunk_latency)  # using 'await asyncio.sleep' instead of 'time.sleep' for StreamingHttpResponse in asgi
             data_chunk = 'd' * chunk_size
+            logger.debug("[%s]sending %d bytes of data", datetime.now().strftime("%H:%M:%S:%f"), chunk_size)
             yield data_chunk
     response = StreamingHttpResponse(data_chunk_generator())
     response['Content-Length'] = chunks * chunk_size
@@ -72,3 +73,8 @@ def redirect_from(request):
 @api_view(['GET'])
 def redirect_to(request):
     return JsonResponse({})
+
+@api_view(['GET'])
+def mock_latency(request, latency):
+    time.sleep(latency)
+    return JsonResponse({})
diff --git a/Engine/Source/Programs/WebTestsServer/requirements.txt b/Engine/Source/Programs/WebTestsServer/requirements.txt
--- a/Engine/Source/Programs/WebTestsServer/requirements.txt
+++ b/Engine/Source/Programs/WebTestsServer/requirements.txt
@@ -1,4 +1,4 @@
-django==4.1.7
+django==4.2.0
 djangorestframework==3.14.0
 channels==4.0.0
 daphne==4.0.0
