Migrate retry system total timeout for http request, by adding total timeout feature in http request;
Deprecate HttpTimeout config, use HttpActivityTimeout or HttpTotalTimeout instead;
Deprecate HttpSendTimeout config, only use HttpActivityTimeout;
Make HttpActivityTimeout work on all platforms, not only CurlHttp;
Added corresponding http tests;
Now because timeout migrated, enable new flow by default in retry system to have non-game thread support.

#tests Tested through WebTests project on all platforms, also tried the game on PC.

diff --git a/Engine/Config/BaseEngine.ini b/Engine/Config/BaseEngine.ini
--- a/Engine/Config/BaseEngine.ini
+++ b/Engine/Config/BaseEngine.ini
@@ -28,10 +28,8 @@ SaveExt=usa
 Port=7777
 
 [HTTP]
-HttpTimeout=180
 HttpConnectionTimeout=60
-HttpReceiveTimeout=30
-HttpSendTimeout=30
+HttpActivityTimeout=180
 ; Flush time limits for FHttpManager::Flush:
 ; Default = when Flush is called with no overriding context
 ; Background = when Flush is called because the app is backgrounding.  This may not be called on backgrounding on all platforms.
diff --git a/Engine/Config/ConfigRedirects.ini b/Engine/Config/ConfigRedirects.ini
--- a/Engine/Config/ConfigRedirects.ini
+++ b/Engine/Config/ConfigRedirects.ini
@@ -53,4 +53,13 @@ EncryptionKey=ClientEncryptionKey
 EncryptionKey=ClientEncryptionKey
 
 [HTTP]
-AllowedDomains=Online.HttpManager https:AllowedDomains
+; The legacy behavior uses HttpTimeout as HttpActivityTimeout, it actually should use HttpActivityTimeout instead. 
+; Now the value of HttpTimeout will be read into HttpActivityTimeout just to keep existing behavior.
+; If it's really an requirement to actually have total timeout for the http request, use HttpTotalTimeout instead.
+; By default HttpTotalTimeout should be 0, means no time out for the entire http request
+HttpTimeout=HttpActivityTimeout
+; Use HttpActivityTimeout instead for checking any activity on the establish connection
+HttpReceiveTimeout=HttpActivityTimeout
+; The legacy behavior doesn't make use of HttpSendTimeout at all, and there is no such support on some platforms.
+; HttpActivityTimeout should be enough to decide if there is still any ongoing activity
+HttpSendTimeout=HttpActivityTimeout
diff --git a/Engine/Source/Programs/AutomationTool/LowLevelTests/Tests/LowLevelTests.WebTestsLowLevelTestsExtension.cs b/Engine/Source/Programs/AutomationTool/LowLevelTests/Tests/LowLevelTests.WebTestsLowLevelTestsExtension.cs
--- a/Engine/Source/Programs/AutomationTool/LowLevelTests/Tests/LowLevelTests.WebTestsLowLevelTestsExtension.cs
+++ b/Engine/Source/Programs/AutomationTool/LowLevelTests/Tests/LowLevelTests.WebTestsLowLevelTestsExtension.cs
@@ -22,7 +22,7 @@ namespace LowLevelTests
 
 		public string ExtraCommandLine(UnrealTargetPlatform InPlatform, string InTestApp, string InBuildPath)
 		{
-			return string.Format("--web_server_ip{0}", UnrealHelpers.GetHostIpAddress());
+			return string.Format("--web_server_ip={0}", UnrealHelpers.GetHostIpAddress());
 		}
 
 		public void PreRunTests()
diff --git a/Engine/Source/Programs/WebTests/Private/TestHttp.cpp b/Engine/Source/Programs/WebTests/Private/TestHttp.cpp
--- a/Engine/Source/Programs/WebTests/Private/TestHttp.cpp
+++ b/Engine/Source/Programs/WebTests/Private/TestHttp.cpp
@@ -41,6 +41,41 @@ class FMockHttpModule : public FHttpModule
 {
 public:
 	using FHttpModule::HttpConnectionTimeout;
+	using FHttpModule::HttpTotalTimeout;
+	using FHttpModule::HttpActivityTimeout;
+};
+
+class FHttpTestLogLevelInitializer
+{
+public:
+	FHttpTestLogLevelInitializer()
+		: OldVerbosity(LogHttp.GetVerbosity())
+	{
+		FParse::Bool(FCommandLine::Get(), TEXT("very_verbose="), bVeryVerbose);
+		if (bVeryVerbose)
+		{
+			LogHttp.SetVerbosity(ELogVerbosity::VeryVerbose);
+		}
+	}
+
+	~FHttpTestLogLevelInitializer()
+	{
+		if (OldVerbosity != LogHttp.GetVerbosity())
+		{
+			LogHttp.SetVerbosity(OldVerbosity);
+		}
+	}
+
+	void DisableWarningsInThisTest()
+	{
+		if (!bVeryVerbose)
+		{
+			LogHttp.SetVerbosity(ELogVerbosity::Error);
+		}
+	}
+
+	bool bVeryVerbose = false;
+	ELogVerbosity::Type OldVerbosity;
 };
 
 class FHttpModuleTestFixture
@@ -51,7 +86,6 @@ public:
 		, WebServerHttpPort(8000)
 		, bRunHeavyTests(false)
 		, bRetryEnabled(true)
-		, OldVerbosity(LogHttp.GetVerbosity())
 	{
 		ParseSettingsFromCommandLine();
 
@@ -69,23 +103,18 @@ public:
 		IModuleInterface* Module = HttpModule;
 		Module->ShutdownModule();
 		delete Module;
-
-		if (OldVerbosity != LogHttp.GetVerbosity())
-		{
-			LogHttp.SetVerbosity(OldVerbosity);
-		}
 	}
 
 	void ParseSettingsFromCommandLine()
 	{
-		FParse::Value(FCommandLine::Get(), TEXT("web_server_ip"), WebServerIp);
-		FParse::Bool(FCommandLine::Get(), TEXT("run_heavy_tests"), bRunHeavyTests);
-		FParse::Bool(FCommandLine::Get(), TEXT("retry_enabled"), bRetryEnabled);
+		FParse::Value(FCommandLine::Get(), TEXT("web_server_ip="), WebServerIp);
+		FParse::Bool(FCommandLine::Get(), TEXT("run_heavy_tests="), bRunHeavyTests);
+		FParse::Bool(FCommandLine::Get(), TEXT("retry_enabled="), bRetryEnabled);
 	}
 
 	void DisableWarningsInThisTest()
 	{
-		LogHttp.SetVerbosity(ELogVerbosity::Error);
+		HttpTestLogLevelInitializer.DisableWarningsInThisTest();
 	}
 
 	const FString UrlWithInvalidPortToTestConnectTimeout() const { return TEXT("http://10.255.255.1:8765"); } // non-routable IP address with a random port
@@ -99,10 +128,9 @@ public:
 	FString WebServerIp;
 	uint32 WebServerHttpPort;
 	FMockHttpModule* HttpModule;
-
 	bool bRunHeavyTests;
 	bool bRetryEnabled;
-	ELogVerbosity::Type OldVerbosity;
+	FHttpTestLogLevelInitializer HttpTestLogLevelInitializer;
 };
 
 TEST_CASE_METHOD(FHttpModuleTestFixture, "Shutdown http module without issue when there are ongoing http requests.", HTTP_TAG)
@@ -244,6 +272,18 @@ TEST_CASE_METHOD(FWaitUntilCompleteHttpFixture, "Http Methods", HTTP_TAG)
 	HttpRequest->ProcessRequest();
 }
 
+TEST_CASE_METHOD(FWaitUntilCompleteHttpFixture, "Can process https request", HTTP_TAG)
+{
+	TSharedRef<IHttpRequest> HttpRequest = HttpModule->CreateRequest();
+	HttpRequest->SetVerb(TEXT("GET"));
+	HttpRequest->SetURL(TEXT("https://www.unrealengine.com/"));
+	HttpRequest->OnProcessRequestComplete().BindLambda([](FHttpRequestPtr HttpRequest, FHttpResponsePtr HttpResponse, bool bSucceeded) {
+		CHECK(bSucceeded);
+		REQUIRE(HttpResponse != nullptr);
+	});
+	HttpRequest->ProcessRequest();
+}
+
 TEST_CASE_METHOD(FWaitUntilCompleteHttpFixture, "Get large response content without chunks", HTTP_TAG)
 {
 	TSharedRef<IHttpRequest> HttpRequest = CreateRequest();
@@ -261,6 +301,7 @@ TEST_CASE_METHOD(FWaitUntilCompleteHttpFixture, "Http request connect timeout",
 {
 	DisableWarningsInThisTest();
 
+	HttpModule->HttpActivityTimeout = 3.0f; // Make sure this won't be triggered before establishing connection
 	HttpModule->HttpConnectionTimeout = 15.0f;
 
 	TSharedRef<IHttpRequest> HttpRequest = CreateRequest();
@@ -617,6 +658,133 @@ void WriteTestFile(const FString& TestFileName, uint64 TestFileSize)
 }
 }
 
+#define HTTP_TEST_TIMEOUT_CHUNK_SIZE 16*1024 // Use a big chunk size so it triggers data received callback in time on all platforms
+
+TEST_CASE_METHOD(FWaitUntilCompleteHttpFixture, "Http request activity timeout", HTTP_TAG)
+{
+	DisableWarningsInThisTest();
+
+	float ReceiveTimeoutSetting = 3.0f;
+	HttpModule->HttpActivityTimeout = ReceiveTimeoutSetting;
+
+	TSharedPtr<IHttpRequest> HttpRequest = CreateRequest();
+	HttpRequest->SetURL(UrlStreamDownload(3/*Chunks*/, HTTP_TEST_TIMEOUT_CHUNK_SIZE, 5/*ChunkLatency*/));
+	HttpRequest->SetVerb(TEXT("GET"));
+
+	const double StartTime = FPlatformTime::Seconds();
+
+	HttpRequest->OnProcessRequestComplete().BindLambda([StartTime, ReceiveTimeoutSetting](FHttpRequestPtr HttpRequest, FHttpResponsePtr HttpResponse, bool bSucceeded) {
+		CHECK(!bSucceeded);
+		CHECK(HttpRequest->GetStatus() == EHttpRequestStatus::Failed);
+		CHECK(HttpRequest->GetFailureReason() == EHttpFailureReason::ConnectionError);
+
+		const double DurationInSeconds  = FPlatformTime::Seconds() - StartTime;
+#if WITH_CURL_XCURL
+		// Unlike libCurl, currently there is an issue in xCurl that it triggers CURLINFO_HEADER_OUT even if can't 
+		// connect. Had to disable that code, make sure not to treat that event as connected
+		// So it takes 5s to receive the first chunk to be considered as connected, then start response timer and 
+		// take 3s to response timeout
+		CHECK(FMath::IsNearlyEqual(DurationInSeconds, ReceiveTimeoutSetting + 5, HTTP_TIME_DIFF_TOLERANCE));
+#else
+		CHECK(FMath::IsNearlyEqual(DurationInSeconds, ReceiveTimeoutSetting, HTTP_TIME_DIFF_TOLERANCE));
+#endif
+
+	});
+	HttpRequest->ProcessRequest();
+}
+
+TEST_CASE_METHOD(FWaitUntilCompleteHttpFixture, "Http request receive won't timeout for streaming request", HTTP_TAG)
+{
+	HttpModule->HttpActivityTimeout = 3.0f;
+
+	TSharedPtr<IHttpRequest> HttpRequest = CreateRequest();
+	HttpRequest->SetURL(UrlStreamDownload(3/*Chunks*/, HTTP_TEST_TIMEOUT_CHUNK_SIZE, 2/*ChunkLatency*/)); // Needs 6s to complete
+	HttpRequest->SetVerb(TEXT("GET"));
+
+	const double StartTime = FPlatformTime::Seconds();
+	HttpRequest->OnProcessRequestComplete().BindLambda([this, StartTime](FHttpRequestPtr HttpRequest, FHttpResponsePtr HttpResponse, bool bSucceeded) {
+		CHECK(bSucceeded);
+		REQUIRE(HttpResponse != nullptr);
+		CHECK(HttpResponse->GetResponseCode() == 200);
+		const double DurationInSeconds = FPlatformTime::Seconds() - StartTime;
+		CHECK(DurationInSeconds > HttpModule->HttpActivityTimeout);
+	});
+	HttpRequest->ProcessRequest();
+}
+
+TEST_CASE_METHOD(FWaitUntilCompleteHttpFixture, "Http request total timeout with get", HTTP_TAG)
+{
+	DisableWarningsInThisTest();
+
+	float TotalTimeoutSetting = 3.0f;
+	HttpModule->HttpTotalTimeout = TotalTimeoutSetting;
+	HttpModule->HttpConnectionTimeout = 5.0f;
+
+	TSharedPtr<IHttpRequest> HttpRequest = CreateRequest();
+	HttpRequest->SetURL(UrlMockLatency(10));
+	HttpRequest->SetVerb(TEXT("GET"));
+
+	const double StartTime = FPlatformTime::Seconds();
+
+	HttpRequest->OnProcessRequestComplete().BindLambda([StartTime, TotalTimeoutSetting](FHttpRequestPtr HttpRequest, FHttpResponsePtr HttpResponse, bool bSucceeded) {
+		CHECK(!bSucceeded);
+		CHECK(HttpRequest->GetStatus() == EHttpRequestStatus::Failed);
+		CHECK(HttpRequest->GetFailureReason() == EHttpFailureReason::TimedOut);
+		const double DurationInSeconds  = FPlatformTime::Seconds() - StartTime;
+		CHECK(FMath::IsNearlyEqual(DurationInSeconds, TotalTimeoutSetting, HTTP_TIME_DIFF_TOLERANCE));
+	});
+	HttpRequest->ProcessRequest();
+}
+
+TEST_CASE_METHOD(FWaitUntilCompleteHttpFixture, "Http request total timeout with streaming download", HTTP_TAG)
+{
+	DisableWarningsInThisTest();
+
+	float TimeoutSetting = 3.0f;
+	HttpModule->HttpActivityTimeout = 2.5f; // Make sure it won't fail because of receive timeout
+	HttpModule->HttpTotalTimeout = TimeoutSetting;
+
+	if (bRetryEnabled)
+	{
+		TimeoutSetting = 4.0f; // This will override http module default timeout
+		HttpRetryManager->RetryTimeoutRelativeSecondsDefault = TimeoutSetting;
+	}
+
+	TSharedPtr<IHttpRequest> HttpRequest;
+	SECTION("Use default timeout from http module or retry manager depends on bRetryEnabled")
+	{
+		HttpRequest = CreateRequest();
+	}
+	SECTION("Override from http request")
+	{
+		TimeoutSetting = 5.0f; // This will override default timeout in http module and retry manager
+
+		if (bRetryEnabled)
+		{
+			HttpRequest = HttpRetryManager->CreateRequest(FHttpRetrySystem::FRetryLimitCountSetting(), TimeoutSetting);
+		}
+		else
+		{
+			HttpRequest = HttpModule->CreateRequest();
+			HttpRequest->SetTimeout(TimeoutSetting);
+		}
+	}
+
+	HttpRequest->SetURL(UrlStreamDownload(4/*Chunks*/, HTTP_TEST_TIMEOUT_CHUNK_SIZE, 2/*ChunkLatency*/)); // Needs 8s to complete
+	HttpRequest->SetVerb(TEXT("GET"));
+
+	const double StartTime = FPlatformTime::Seconds();
+
+	HttpRequest->OnProcessRequestComplete().BindLambda([StartTime, TimeoutSetting](FHttpRequestPtr HttpRequest, FHttpResponsePtr HttpResponse, bool bSucceeded) {
+		CHECK(!bSucceeded);
+		CHECK(HttpRequest->GetStatus() == EHttpRequestStatus::Failed);
+		CHECK(HttpRequest->GetFailureReason() == EHttpFailureReason::TimedOut);
+		const double DurationInSeconds  = FPlatformTime::Seconds() - StartTime;
+		CHECK(FMath::IsNearlyEqual(DurationInSeconds, TimeoutSetting, HTTP_TIME_DIFF_TOLERANCE));
+	});
+	HttpRequest->ProcessRequest();
+}
+
 TEST_CASE_METHOD(FWaitUntilCompleteHttpFixture, "Streaming http upload from file by PUT can work well", HTTP_TAG)
 {
 	FString Filename = FString(FPlatformProcess::UserSettingsDir()) / TEXT("TestStreamUpload.dat");
diff --git a/Engine/Source/Programs/WebTests/Private/TestWebSockets.cpp b/Engine/Source/Programs/WebTests/Private/TestWebSockets.cpp
--- a/Engine/Source/Programs/WebTests/Private/TestWebSockets.cpp
+++ b/Engine/Source/Programs/WebTests/Private/TestWebSockets.cpp
@@ -66,8 +66,8 @@ public:
 
 	void ParseSettingsFromCommandLine()
 	{
-		FParse::Value(FCommandLine::Get(), TEXT("web_server_ip"), WebServerIp);
-		FParse::Value(FCommandLine::Get(), TEXT("web_server_websockets_port"), WebServerWebSocketsPort);
+		FParse::Value(FCommandLine::Get(), TEXT("web_server_ip="), WebServerIp);
+		FParse::Value(FCommandLine::Get(), TEXT("web_server_websockets_port="), WebServerWebSocketsPort);
 	}
 
 	void DisableWarningsInThisTest()
diff --git a/Engine/Source/Runtime/Online/HTTP/Private/Apple/AppleHttp.cpp b/Engine/Source/Runtime/Online/HTTP/Private/Apple/AppleHttp.cpp
--- a/Engine/Source/Runtime/Online/HTTP/Private/Apple/AppleHttp.cpp
+++ b/Engine/Source/Runtime/Online/HTTP/Private/Apple/AppleHttp.cpp
@@ -410,7 +410,7 @@ FAppleHttpRequest::FAppleHttpRequest(NSURLSession* InSession)
 FAppleHttpRequest::~FAppleHttpRequest()
 {
 	UE_LOG(LogHttp, Verbose, TEXT("FAppleHttpRequest::~FAppleHttpRequest()"));
-	CleanupRequest();
+	PostProcess();
 	[Request release];
     [Session release];
 }
@@ -691,7 +691,7 @@ bool FAppleHttpRequest::SetupRequest()
 		[Request setValue:Tag forHTTPHeaderField:@"User-Agent"];
 	}
 
-	CleanupRequest();
+	PostProcess();
 
 	LastReportedBytesWritten = 0;
 	LastReportedBytesRead = 0;
@@ -732,8 +732,7 @@ void FAppleHttpRequest::FinishRequest()
 {
 	UE_LOG(LogHttp, Verbose, TEXT("FAppleHttpRequest::FinishRequest()"));
 
-	// Clean up session/request handles that may have been created
-	CleanupRequest();
+	PostProcess();
 
 	bool bSucceeded = (Response && Response->GetStatusFromDelegate() == EHttpRequestStatus::Succeeded);
 	UE_LOG(LogHttp, Verbose, TEXT("Request %s"), bSucceeded ? TEXT("succeeded") : TEXT("failed"));
@@ -741,7 +740,16 @@ void FAppleHttpRequest::FinishRequest()
 
 	if (!bSucceeded)
 	{
-		SetFailureReason(Response ? Response->GetFailureReasonFromDelegate() : EHttpFailureReason::Other);
+		EHttpFailureReason Reason = EHttpFailureReason::Other;
+		if (Response)
+		{
+			Reason = Response->GetFailureReasonFromDelegate();
+			if (Reason == EHttpFailureReason::Cancelled && bTimedOut)
+			{
+				Reason = EHttpFailureReason::TimedOut;
+			}
+		}
+		SetFailureReason(Reason);
 
 		if (GetFailureReason() == EHttpFailureReason::ConnectionError)
 		{
@@ -777,10 +785,8 @@ void FAppleHttpRequest::CleanupRequest()
 	}
 }
 
-void FAppleHttpRequest::CancelRequest()
+void FAppleHttpRequest::AbortRequest()
 {
-	UE_LOG(LogHttp, Verbose, TEXT("FAppleHttpRequest::CancelRequest()"));
-
 	if (Task != nil)
 	{
 		[Task cancel];
diff --git a/Engine/Source/Runtime/Online/HTTP/Private/Apple/AppleHttp.h b/Engine/Source/Runtime/Online/HTTP/Private/Apple/AppleHttp.h
--- a/Engine/Source/Runtime/Online/HTTP/Private/Apple/AppleHttp.h
+++ b/Engine/Source/Runtime/Online/HTTP/Private/Apple/AppleHttp.h
@@ -46,7 +46,6 @@ public:
 	virtual void SetHeader(const FString& HeaderName, const FString& HeaderValue) override;
 	virtual void AppendToHeader(const FString& HeaderName, const FString& AdditionalHeaderValue) override;
 	virtual bool ProcessRequest() override;
-	virtual void CancelRequest() override;
 	virtual const FHttpResponsePtr GetResponse() const override;
 	virtual void Tick(float DeltaSeconds) override;
 	virtual float GetElapsedTime() const override;
@@ -88,10 +87,12 @@ private:
 	 */
 	virtual bool SetupRequest() override;
 
+	virtual void AbortRequest() override;
+
 	/**
 	 * Close session/request handles and unregister callbacks
 	 */
-	void CleanupRequest();
+	virtual void CleanupRequest() override;
 
 private:
 	/** This is the NSMutableURLRequest, all our Apple functionality will deal with this. */
diff --git a/Engine/Source/Runtime/Online/HTTP/Private/Apple/ApplePlatformHttp.cpp b/Engine/Source/Runtime/Online/HTTP/Private/Apple/ApplePlatformHttp.cpp
--- a/Engine/Source/Runtime/Online/HTTP/Private/Apple/ApplePlatformHttp.cpp
+++ b/Engine/Source/Runtime/Online/HTTP/Private/Apple/ApplePlatformHttp.cpp
@@ -222,7 +222,7 @@ void FApplePlatformHttp::Init()
 	FSslModule::Get();
 #endif
 
-	InitWithNSUrlSession();
+	// Lazy init, call InitWithNSUrlSession when need to create request, so session config can be set before creating session
 }
 
 void FApplePlatformHttp::Shutdown()
@@ -236,6 +236,10 @@ void FApplePlatformHttp::InitWithNSUrlSession()
 
 	// Disable cache to mimic WinInet behavior
 	Config.requestCachePolicy = NSURLRequestReloadIgnoringLocalCacheData;
+
+	float HttpActivityTimeout = FHttpModule::Get().GetHttpActivityTimeout();
+	check(HttpActivityTimeout > 0);
+	Config.timeoutIntervalForRequest = HttpActivityTimeout;
 	
 #if WITH_SSL
 	// Load SSL module during HTTP module's StatupModule() to make sure module manager figures out the dependencies correctly
@@ -267,6 +271,11 @@ FHttpManager* FApplePlatformHttp::CreatePlatformHttpManager()
 
 IHttpRequest* FApplePlatformHttp::ConstructRequest()
 {
+	if (Session == nil)
+	{
+		InitWithNSUrlSession();
+	}
+
 	return new FAppleHttpRequest(Session);
 }
 
diff --git a/Engine/Source/Runtime/Online/HTTP/Private/Curl/CurlHttp.cpp b/Engine/Source/Runtime/Online/HTTP/Private/Curl/CurlHttp.cpp
--- a/Engine/Source/Runtime/Online/HTTP/Private/Curl/CurlHttp.cpp
+++ b/Engine/Source/Runtime/Online/HTTP/Private/Curl/CurlHttp.cpp
@@ -5,7 +5,6 @@
 #include "Curl/CurlHttp.h"
 #include "Stats/Stats.h"
 #include "Misc/App.h"
-#include "HAL/PlatformTime.h"
 #include "HttpModule.h"
 #include "Http.h"
 #include "Misc/EngineVersion.h"
@@ -66,13 +65,11 @@ FCurlHttpRequest::FCurlHttpRequest()
 	: EasyHandle(nullptr)
 	, HeaderList(nullptr)
 	, Verb(TEXT("GET"))
-	, bCanceled(false)
 	, bCurlRequestCompleted(false)
 	, bRedirected(false)
 	, CurlAddToMultiResult(CURLM_OK)
 	, CurlCompletionResult(CURLE_OK)
 	, ElapsedTime(0.0f)
-	, TimeSinceLastResponse(0.0f)
 	, bAnyHttpActivity(false)
 	, BytesSent(0)
 	, TotalBytesSent(0)
@@ -181,6 +178,8 @@ FCurlHttpRequest::FCurlHttpRequest()
 	{
 		SetHeader(It.Key(), It.Value());
 	}
+
+	bUsePlatformActivityTimeout = false;
 }
 
 FCurlHttpRequest::~FCurlHttpRequest()
@@ -470,7 +469,8 @@ size_t FCurlHttpRequest::ReceiveResponseHeaderCallback(void* Ptr, size_t SizeInB
 		TotalBytesRead = 0;
 	}
 
-	TimeSinceLastResponse = 0.0f;
+	OnAnyActivityOccur(TEXTVIEW("Received header"));
+
 	uint32 HeaderSize = SizeInBlocks * BlockSizeInBytes;
 	if (HeaderSize > 0 && HeaderSize <= CURL_MAX_HTTP_HEADER)
 	{
@@ -555,7 +555,7 @@ size_t FCurlHttpRequest::ReceiveResponseBodyCallback(void* Ptr, size_t SizeInBlo
 		Response = MakeShared<FCurlHttpResponse, ESPMode::ThreadSafe>(*this);
 	}
 
-	TimeSinceLastResponse = 0.0f;
+	OnAnyActivityOccur(TEXTVIEW("Received body"));
 
 	// Number of bytes actually taken care of. If that amount differs from the amount passed to your 
 	// callback function, it will signal an error condition to the library. This will cause the transfer 
@@ -599,7 +599,7 @@ size_t FCurlHttpRequest::ReceiveResponseBodyCallback(void* Ptr, size_t SizeInBlo
 
 size_t FCurlHttpRequest::UploadCallback(void* Ptr, size_t SizeInBlocks, size_t BlockSizeInBytes)
 {
-	TimeSinceLastResponse = 0.0f;
+	OnAnyActivityOccur(TEXTVIEW("Upload callback"));
 
 	size_t MaxBufferSize = SizeInBlocks * BlockSizeInBytes;
 	size_t SizeAlreadySent = BytesSent.load();
@@ -743,30 +743,51 @@ size_t FCurlHttpRequest::DebugCallback(CURL * Handle, curl_infotype DebugInfoTyp
 
 	switch (DebugInfoType)
 	{
-		case CURLINFO_HEADER_IN:
-		case CURLINFO_HEADER_OUT:
-		case CURLINFO_DATA_IN:
-		case CURLINFO_DATA_OUT:
-		case CURLINFO_SSL_DATA_IN:
-		case CURLINFO_SSL_DATA_OUT:
-			TimeSinceLastResponse = 0.0f;
-			bAnyHttpActivity = true;
+	case CURLINFO_HEADER_IN: 
+		OnAnyActivityOccur(TEXTVIEW("Header in"));
+		break;
+	case CURLINFO_HEADER_OUT: 
+		// Unlike libCurl, currently there is an issue in xCurl that it triggers CURLINFO_HEADER_OUT even if can't 
+		// connect. Had to disable this code, make sure not to treat that event as connected/activity happened
+#if !WITH_CURL_XCURL
+		OnAnyActivityOccur(TEXTVIEW("Header out"));
+#endif
+		break;
+	case CURLINFO_DATA_IN: 
+		OnAnyActivityOccur(TEXTVIEW("Data in"));
+		break;
+	case CURLINFO_DATA_OUT:
+		OnAnyActivityOccur(TEXTVIEW("Data out"));
+		break;
+	case CURLINFO_SSL_DATA_IN:
+		OnAnyActivityOccur(TEXTVIEW("Ssl data in"));
+		break;
+	case CURLINFO_SSL_DATA_OUT:
+		OnAnyActivityOccur(TEXTVIEW("Ssl data out"));
+		break;
+	default:
+		break;
+	}
+
+	return 0;
+}
+
+void FCurlHttpRequest::OnAnyActivityOccur(FStringView Reason)
+{
+	if (!bAnyHttpActivity)
+	{
+		bAnyHttpActivity = true;
+
 #if WITH_CURL_XCURL
-			// Unlike libCurl, currently there is an issue in xCurl that it triggers CURLINFO_HEADER_OUT even if can't 
-			// connect. Had to disable this code, make sure not to treat that event as connected
-			if (ConnectTime < 0 && DebugInfoType != CURLINFO_HEADER_OUT)
-			{
-				ConnectTime = FPlatformTime::Seconds() - StartProcessTime;
-			}
+		ConnectTime = FPlatformTime::Seconds() - StartProcessTime;
 #else
-			curl_easy_getinfo(EasyHandle, CURLINFO_CONNECT_TIME, &ConnectTime);
+		curl_easy_getinfo(EasyHandle, CURLINFO_CONNECT_TIME, &ConnectTime);
 #endif
-			break;
-		default:
-			break;
+
+		StartActivityTimeoutTimer();
 	}
-	
-	return 0;
+
+	ResetActivityTimeoutTimer(Reason);
 }
 
 bool FCurlHttpRequest::SetupRequest()
@@ -790,7 +811,6 @@ bool FCurlHttpRequest::SetupRequest()
 	}
 
 	bCurlRequestCompleted = false;
-	bCanceled = false;
 	CurlAddToMultiResult = CURLM_OK;
 	LastReportedBytesSent = 0;
 
@@ -996,22 +1016,12 @@ bool FCurlHttpRequest::ProcessRequest()
 	// Clear out response. If this is a re-used request, Response could point to a stale response until SetupRequestHttpThread is called
 	Response = nullptr;
 	LastReportedBytesRead = 0;
-	TotalBytesRead = 0;
 
 	if (!PreProcess())
 	{
 		return false;
 	}
 
-	// Clear the info cache log so we don't output messages from previous requests when reusing/retrying a request
-	{
-		const FScopeLock CacheLock(&InfoMessageCacheCriticalSection);
-		for (FString& Line : InfoMessageCache)
-		{
-			Line.Reset();
-		}
-	}
-
 	QUICK_SCOPE_CYCLE_COUNTER(STAT_CurlHttpAddThreadedRequest);
 	// Mark as in-flight to prevent overlapped requests using the same object
 	SetStatus(EHttpRequestStatus::Processing);
@@ -1027,12 +1037,30 @@ bool FCurlHttpRequest::ProcessRequest()
 	return true;
 }
 
+void FCurlHttpRequest::ClearInCaseOfRetry()
+{
+	IHttpThreadedRequest::ClearInCaseOfRetry();
+
+	// Clear out response. If this is a re-used request, Response could point to a stale response until SetupRequestHttpThread is called
+	Response = nullptr;
+	LastReportedBytesRead = 0;
+	TotalBytesRead = 0;
+	bAnyHttpActivity = false;
+
+	// Clear the info cache log so we don't output messages from previous requests when reusing/retrying a request
+	{
+		const FScopeLock CacheLock(&InfoMessageCacheCriticalSection);
+		for (FString& Line : InfoMessageCache)
+		{
+			Line.Reset();
+		}
+	}
+}
+
 bool FCurlHttpRequest::StartThreadedRequest()
 {
 	// reset timeout
 	ElapsedTime = 0.0f;
-	TimeSinceLastResponse = 0.0f;
-	bAnyHttpActivity = false;
 	
 	UE_LOG(LogHttp, Verbose, TEXT("%p: request (easy handle:%p) has started threaded processing"), this, EasyHandle);
 
@@ -1045,26 +1073,19 @@ bool FCurlHttpRequest::IsThreadedRequestComplete()
 	{
 		return true;
 	}
-	
-	if (bCurlRequestCompleted && ElapsedTime >= FHttpModule::Get().GetHttpDelayTime())
+
+	if (bTimedOut || bActivityTimedOut)
 	{
 		return true;
 	}
-
-	if (CurlAddToMultiResult != CURLM_OK)
+	
+	if (bCurlRequestCompleted  && ElapsedTime >= FHttpModule::Get().GetHttpDelayTime())
 	{
 		return true;
 	}
 
-	const float HttpTimeout = GetTimeoutOrDefault();
-	bool bTimedOut = (HttpTimeout > 0 && TimeSinceLastResponse >= HttpTimeout);
-#if CURL_ENABLE_NO_TIMEOUTS_OPTION
-	static const bool bNoTimeouts = FParse::Param(FCommandLine::Get(), TEXT("NoTimeouts"));
-	bTimedOut = bTimedOut && !bNoTimeouts;
-#endif
-	if (bTimedOut)
+	if (CurlAddToMultiResult != CURLM_OK)
 	{
-		UE_LOG(LogHttp, Warning, TEXT("%p: HTTP request timed out after %0.2f seconds URL=%s"), this, TimeSinceLastResponse, *GetURL());
 		return true;
 	}
 
@@ -1074,21 +1095,10 @@ bool FCurlHttpRequest::IsThreadedRequestComplete()
 void FCurlHttpRequest::TickThreadedRequest(float DeltaSeconds)
 {
 	ElapsedTime += DeltaSeconds;
-	TimeSinceLastResponse += DeltaSeconds;
 }
 
-void FCurlHttpRequest::CancelRequest()
+void FCurlHttpRequest::AbortRequest()
 {
-	QUICK_SCOPE_CYCLE_COUNTER(STAT_FCurlHttpRequest_CancelRequest);
-
-	if (bCanceled)
-	{
-		return;
-	}
-
-	bCanceled = true;
-	UE_LOG(LogHttp, Verbose, TEXT("%p: HTTP request canceled.  URL=%s"), this, *GetURL());
-
 	FHttpManager& HttpManager = FHttpModule::Get().GetHttpManager();
 	if (HttpManager.IsValidRequest(this))
 	{
@@ -1150,18 +1160,15 @@ void FCurlHttpRequest::MarkAsCompleted(CURLcode InCurlCompletionResult)
 {
 	CurlCompletionResult = InCurlCompletionResult;
 	bCurlRequestCompleted = true;
+
+	StopActivityTimeoutTimer();
 }
 
 void FCurlHttpRequest::FinishRequest()
 {
 	QUICK_SCOPE_CYCLE_COUNTER(STAT_FCurlHttpRequest_FinishRequest);
 
-	if (RequestPayload.IsValid())
-	{
-		RequestPayload->Close();
-	}
-
-	curl_easy_setopt(EasyHandle, CURLOPT_SHARE, nullptr);
+	PostProcess();
 	
 	CheckProgressDelegate();
 	// if completed, get more info
@@ -1288,6 +1295,14 @@ void FCurlHttpRequest::FinishRequest()
 		{
 			SetFailureReason(EHttpFailureReason::Cancelled);
 		}
+		else if (bTimedOut)
+		{
+			SetFailureReason(EHttpFailureReason::TimedOut);
+		}
+		else if (bActivityTimedOut)
+		{
+			SetFailureReason(EHttpFailureReason::ConnectionError);
+		}
 		else if (bCurlRequestCompleted)
 		{
 			switch (CurlCompletionResult)
@@ -1309,14 +1324,7 @@ void FCurlHttpRequest::FinishRequest()
 		}
 		else
 		{
-			if (bAnyHttpActivity)
-			{
-				SetFailureReason(EHttpFailureReason::Other);
-			}
-			else
-			{
-				SetFailureReason(EHttpFailureReason::ConnectionError);
-			}
+			SetFailureReason(EHttpFailureReason::Other);
 		}
 		// Call delegate with failure
 		OnProcessRequestComplete().ExecuteIfBound(SharedThis(this), Response, false);
@@ -1332,6 +1340,15 @@ float FCurlHttpRequest::GetElapsedTime() const
 	return ElapsedTime;
 }
 
+void FCurlHttpRequest::CleanupRequest()
+{
+	if (RequestPayload.IsValid())
+	{
+		RequestPayload->Close();
+	}
+
+	curl_easy_setopt(EasyHandle, CURLOPT_SHARE, nullptr);
+}
 
 // FCurlHttpRequest
 
diff --git a/Engine/Source/Runtime/Online/HTTP/Private/Curl/CurlHttp.h b/Engine/Source/Runtime/Online/HTTP/Private/Curl/CurlHttp.h
--- a/Engine/Source/Runtime/Online/HTTP/Private/Curl/CurlHttp.h
+++ b/Engine/Source/Runtime/Online/HTTP/Private/Curl/CurlHttp.h
@@ -165,7 +165,6 @@ public:
 	virtual void SetHeader(const FString& HeaderName, const FString& HeaderValue) override;
 	virtual void AppendToHeader(const FString& HeaderName, const FString& AdditionalHeaderValue) override;
 	virtual bool ProcessRequest() override;
-	virtual void CancelRequest() override;
 	virtual const FHttpResponsePtr GetResponse() const override;
 	virtual void Tick(float DeltaSeconds) override;
 	virtual float GetElapsedTime() const override;
@@ -339,6 +338,8 @@ private:
 	 */
 	virtual bool SetupRequest() override;
 
+	virtual void AbortRequest() override;
+
 	/**
 	 * Trigger the request progress delegate if progress has changed
 	 */
@@ -349,7 +350,13 @@ private:
 
 	/** Combine a header's key/value in the format "Key: Value" */
 	static FString CombineHeaderKeyValue(const FString& HeaderKey, const FString& HeaderValue);
-	
+
+	virtual void CleanupRequest() override;
+
+	void OnAnyActivityOccur(FStringView Reason);
+
+	void ClearInCaseOfRetry() override;
+
 private:
 
 	/** Pointer to an easy handle specific to this request */
@@ -360,8 +367,6 @@ private:
 	FString			URL;
 	/** Cached verb */
 	FString			Verb;
-	/** Set to true if request has been canceled */
-	bool			bCanceled;
 	/** Set to true when request has been completed */
 	std::atomic<bool> bCurlRequestCompleted;
 	/** Set to true when request has "30* Multiple Choices" (e.g. 301 Moved Permanently, 302 temporary redirect, 308 Permanent Redirect, etc.) */
@@ -382,8 +387,6 @@ private:
 	TMap<FString, FString> Headers;
 	/** Total elapsed time in seconds since the start of the request */
 	float ElapsedTime;
-	/** Elapsed time since the last received HTTP response. */
-	float TimeSinceLastResponse;
 	/** Have we had any HTTP activity with the host? Sending headers, SSL handshake, etc */
 	bool bAnyHttpActivity;
 	/** Newly received headers we need to inform listeners about */
diff --git a/Engine/Source/Runtime/Online/HTTP/Private/EventLoopHttpThread.cpp b/Engine/Source/Runtime/Online/HTTP/Private/EventLoopHttpThread.cpp
--- a/Engine/Source/Runtime/Online/HTTP/Private/EventLoopHttpThread.cpp
+++ b/Engine/Source/Runtime/Online/HTTP/Private/EventLoopHttpThread.cpp
@@ -8,6 +8,24 @@
 #include "PlatformHttp.h"
 #include "Stats/Stats.h"
 
+class FHttpTaskTimerHandleEventLoop : public IHttpTaskTimerHandle
+{
+public:
+	FHttpTaskTimerHandleEventLoop(UE::EventLoop::FTimerHandle InHandle)
+		: Handle(InHandle)
+	{
+	}
+
+	virtual void RemoveTaskFrom(FHttpThreadBase* HttpThreadBase) 
+	{ 
+		HttpThreadBase->RemoveTimerHandle(Handle);
+	}
+
+private:
+	UE::EventLoop::FTimerHandle Handle;
+};
+
+
 FEventLoopHttpThread::FEventLoopHttpThread()
 {
 	FPlatformHttp::AddDefaultUserAgentProjectComment(TEXT("http-eventloop"));
@@ -161,8 +179,19 @@ void FEventLoopHttpThread::ResetTickTimer()
 	true /* repeat */);
 }
 
-void FEventLoopHttpThread::AddHttpThreadTask(TFunction<void()>&& Task, float InDelay)
+TSharedPtr<IHttpTaskTimerHandle> FEventLoopHttpThread::AddHttpThreadTask(TFunction<void()>&& Task, float InDelay)
 {
 	UE::EventLoop::IEventLoop& EventLoop = GetEventLoopChecked();
-	EventLoop.SetTimer(Task, FTimespan::FromSeconds(InDelay));
-}
\ No newline at end of file
+	return MakeShared<FHttpTaskTimerHandleEventLoop>(EventLoop.SetTimer(Task, FTimespan::FromSeconds(InDelay)));
+}
+
+void FEventLoopHttpThread::RemoveTimerHandle(FTSTicker::FDelegateHandle DelegateHandle)
+{
+	checkNoEntry();
+}
+
+void FEventLoopHttpThread::RemoveTimerHandle(UE::EventLoop::FTimerHandle EventLoopTimerHandle)
+{
+	UE::EventLoop::IEventLoop& EventLoop = GetEventLoopChecked();
+	EventLoop.ClearTimer(EventLoopTimerHandle);
+}
diff --git a/Engine/Source/Runtime/Online/HTTP/Private/EventLoopHttpThread.h b/Engine/Source/Runtime/Online/HTTP/Private/EventLoopHttpThread.h
--- a/Engine/Source/Runtime/Online/HTTP/Private/EventLoopHttpThread.h
+++ b/Engine/Source/Runtime/Online/HTTP/Private/EventLoopHttpThread.h
@@ -40,7 +40,11 @@ protected:
 
 	void ResetTickTimer();
 
-	virtual void AddHttpThreadTask(TFunction<void()>&& Task, float InDelay) override;
+	virtual TSharedPtr<IHttpTaskTimerHandle> AddHttpThreadTask(TFunction<void()>&& Task, float InDelay) override;
+
+	virtual void RemoveTimerHandle(FTSTicker::FDelegateHandle DelegateHandle) override;
+
+	virtual void RemoveTimerHandle(UE::EventLoop::FTimerHandle EventLoopTimerHandle) override;
 
 protected:
 	UE::EventLoop::FTimerHandle RequestTickTimer;
diff --git a/Engine/Source/Runtime/Online/HTTP/Private/GenericPlatform/HttpRequestCommon.cpp b/Engine/Source/Runtime/Online/HTTP/Private/GenericPlatform/HttpRequestCommon.cpp
--- a/Engine/Source/Runtime/Online/HTTP/Private/GenericPlatform/HttpRequestCommon.cpp
+++ b/Engine/Source/Runtime/Online/HTTP/Private/GenericPlatform/HttpRequestCommon.cpp
@@ -4,6 +4,14 @@
 #include "GenericPlatform/HttpResponseCommon.h"
 #include "Http.h"
 #include "HttpManager.h"
+#include "Misc/CommandLine.h"
+#include "Stats/Stats.h"
+
+FHttpRequestCommon::FHttpRequestCommon()
+	: RequestStartTimeAbsoluteSeconds(FPlatformTime::Seconds())
+	, ActivityTimeoutAt(0.0)
+{
+}
 
 FString FHttpRequestCommon::GetURLParameter(const FString& ParameterName) const
 {
@@ -60,6 +68,12 @@ bool FHttpRequestCommon::PreCheck() const
 		return false;
 	}
 
+	if (bTimedOut)
+	{
+		UE_LOG(LogHttp, Warning, TEXT("ProcessRequest failed. Request with URL '%s' already timed out."), *GetURL());
+		return false;
+	}
+
 	return true;
 }
 
@@ -73,16 +87,25 @@ bool FHttpRequestCommon::PreProcess()
 		return false;
 	}
 
+	StartTotalTimeoutTimer();
+
 	UE_LOG(LogHttp, Verbose, TEXT("%p: Verb='%s' URL='%s'"), this, *GetVerb(), *GetURL());
 
 	return true;
 }
 
+void FHttpRequestCommon::PostProcess()
+{
+	CleanupRequest();
+}
+
 void FHttpRequestCommon::ClearInCaseOfRetry()
 {
 	// TODO: clear response shared ptr here as well after moving it from child class to this class
 
+	bActivityTimedOut = false;
 	FailureReason = EHttpFailureReason::None;
+	bCanceled = false;
 }
 
 void FHttpRequestCommon::FinishRequestNotInHttpManager()
@@ -174,7 +197,189 @@ TOptional<float> FHttpRequestCommon::GetTimeout() const
 
 float FHttpRequestCommon::GetTimeoutOrDefault() const
 {
-	return GetTimeout().Get(FHttpModule::Get().GetHttpTimeout());
+	return GetTimeout().Get(FHttpModule::Get().GetHttpTotalTimeout());
+}
+
+void FHttpRequestCommon::CancelRequest()
+{
+	if (bCanceled)
+	{
+		return;
+	}
+
+	bCanceled = true;
+	UE_LOG(LogHttp, Verbose, TEXT("HTTP request canceled. URL=%s"), *GetURL());
+
+	FHttpModule::Get().GetHttpManager().AddHttpThreadTask([StrongThis = StaticCastSharedRef<FHttpRequestCommon>(AsShared())]()
+	{
+		// Run AbortRequest in HTTP thread to avoid potential concurrency issue
+		QUICK_SCOPE_CYCLE_COUNTER(STAT_FHttpRequestCommon_AbortRequest);
+		StrongThis->AbortRequest();
+	});
+}
+
+void FHttpRequestCommon::StartActivityTimeoutTimer()
+{
+	const FScopeLock CacheLock(&HttpTaskTimerHandleCriticalSection);
+
+	if (bUsePlatformActivityTimeout)
+	{
+		return;
+	}
+
+#if !UE_BUILD_SHIPPING
+	static const bool bNoTimeouts = FParse::Param(FCommandLine::Get(), TEXT("NoTimeouts"));
+	if (bNoTimeouts)
+	{
+		return;
+	}
+#endif
+
+	if (bActivityTimedOut)
+	{
+		return;
+	}
+
+	float HttpActivityTimeout = FHttpModule::Get().GetHttpActivityTimeout();
+	check(HttpActivityTimeout > 0);
+	StartActivityTimeoutTimerBy(HttpActivityTimeout);
+
+	ResetActivityTimeoutTimer(TEXTVIEW("Connected"));
+}
+
+void FHttpRequestCommon::StartActivityTimeoutTimerBy(double DelayToTrigger)
+{
+	check(ActivityTimeoutHttpTaskTimerHandle == nullptr);
+
+	TWeakPtr<IHttpRequest> RequestWeakPtr(AsShared());
+	ActivityTimeoutHttpTaskTimerHandle = FHttpModule::Get().GetHttpManager().AddHttpThreadTask([RequestWeakPtr]() {
+		if (TSharedPtr<IHttpRequest> RequestPtr = RequestWeakPtr.Pin())
+		{
+			TSharedPtr<FHttpRequestCommon> RequestCommonPtr = StaticCastSharedPtr<FHttpRequestCommon>(RequestPtr);
+			RequestCommonPtr->OnActivityTimeoutTimerTaskTrigger();
+		}
+	}, DelayToTrigger + 0.05);
+}
+
+void FHttpRequestCommon::OnActivityTimeoutTimerTaskTrigger()
+{
+	const FScopeLock CacheLock(&HttpTaskTimerHandleCriticalSection);
+
+	ActivityTimeoutHttpTaskTimerHandle.Reset();
+
+	checkf(!EHttpRequestStatus::IsFinished(GetStatus()), TEXT("PostProcess must be called when request complete, to stop activity timeout timer."));
+
+	if (FPlatformTime::Seconds() < ActivityTimeoutAt)
+	{
+		// Check back later
+		UE_LOG(LogHttp, VeryVerbose, TEXT("Request %p check response timeout at [%s], will check again in %.5f seconds"), this, *FDateTime::Now().ToString(TEXT("%H:%M:%S:%s")), ActivityTimeoutAt - FPlatformTime::Seconds());
+		StartActivityTimeoutTimerBy(ActivityTimeoutAt - FPlatformTime::Seconds());
+		return;
+	}
+
+	QUICK_SCOPE_CYCLE_COUNTER(STAT_FHttpRequestCommon_AbortRequest);
+	bActivityTimedOut = true;
+	AbortRequest();
+	UE_LOG(LogHttp, Warning, TEXT("Request [%s] timed out at [%s] because of no responding for %0.2f seconds"), *GetURL(), *FDateTime::Now().ToString(TEXT("%H:%M:%S:%s")), FHttpModule::Get().GetHttpActivityTimeout());
+}
+
+void FHttpRequestCommon::ResetActivityTimeoutTimer(FStringView Reason)
+{
+	const FScopeLock CacheLock(&HttpTaskTimerHandleCriticalSection);
+
+	if (bUsePlatformActivityTimeout)
+	{
+		return;
+	}
+
+	ActivityTimeoutAt = FPlatformTime::Seconds() + FHttpModule::Get().GetHttpActivityTimeout();
+	UE_LOG(LogHttp, VeryVerbose, TEXT("Request [%p] reset response timeout timer at %s: %s"), this, *FDateTime::Now().ToString(TEXT("%H:%M:%S:%s")), Reason.GetData());
+}
+
+void FHttpRequestCommon::StopActivityTimeoutTimer()
+{
+	const FScopeLock CacheLock(&HttpTaskTimerHandleCriticalSection);
+
+	if (bUsePlatformActivityTimeout)
+	{
+		return;
+	}
+
+	if (ActivityTimeoutHttpTaskTimerHandle)
+	{
+		FHttpModule::Get().GetHttpManager().RemoveHttpThreadTask(ActivityTimeoutHttpTaskTimerHandle);
+		ActivityTimeoutHttpTaskTimerHandle.Reset();
+	}
+}
+
+void FHttpRequestCommon::StartTotalTimeoutTimer()
+{
+	const FScopeLock CacheLock(&HttpTaskTimerHandleCriticalSection);
+
+#if !UE_BUILD_SHIPPING
+	static const bool bNoTimeouts = FParse::Param(FCommandLine::Get(), TEXT("NoTimeouts"));
+	if (bNoTimeouts)
+	{
+		return;
+	}
+#endif
+
+	float TimeoutOrDefault = GetTimeoutOrDefault();
+	if (TimeoutOrDefault == 0)
+	{
+		return;
+	}
+
+	if (bTimedOut)
+	{
+		return;
+	}
+
+	// Timeout include retries, so if it's already started before, check this to prevent from adding timer multiple times
+	if (TotalTimeoutHttpTaskTimerHandle)
+	{
+		return;
+	}
+
+	TWeakPtr<IHttpRequest> RequestWeakPtr(AsShared());
+	TotalTimeoutHttpTaskTimerHandle = FHttpModule::Get().GetHttpManager().AddHttpThreadTask([RequestWeakPtr]() {
+		if (TSharedPtr<IHttpRequest> RequestPtr = RequestWeakPtr.Pin())
+		{
+			TSharedPtr<FHttpRequestCommon> RequestCommonPtr = StaticCastSharedPtr<FHttpRequestCommon>(RequestPtr);
+			RequestCommonPtr->OnTotalTimeoutTimerTaskTrigger();
+		}
+	}, TimeoutOrDefault);
+}
+
+void FHttpRequestCommon::OnTotalTimeoutTimerTaskTrigger()
+{
+	const FScopeLock CacheLock(&HttpTaskTimerHandleCriticalSection);
+	bTimedOut = true;
+	UE_LOG(LogHttp, Warning, TEXT("HTTP request timed out after %0.2f seconds URL=%s"), GetTimeoutOrDefault(), *GetURL());
+
+	if (!EHttpRequestStatus::IsFinished(GetStatus())) 
+	{
+		QUICK_SCOPE_CYCLE_COUNTER(STAT_FHttpRequestCommon_AbortRequest);
+		AbortRequest();
+	}
+}
+
+void FHttpRequestCommon::StopTotalTimeoutTimer()
+{
+	const FScopeLock CacheLock(&HttpTaskTimerHandleCriticalSection);
+
+	if (TotalTimeoutHttpTaskTimerHandle)
+	{
+		FHttpModule::Get().GetHttpManager().RemoveHttpThreadTask(TotalTimeoutHttpTaskTimerHandle);
+		TotalTimeoutHttpTaskTimerHandle.Reset();
+	}
+}
+void FHttpRequestCommon::Shutdown()
+{
+	FHttpRequestImpl::Shutdown();
+
+	StopActivityTimeoutTimer();
+	StopTotalTimeoutTimer();
 }
 
 void FHttpRequestCommon::TriggerStatusCodeReceivedDelegate(int32 StatusCode)
diff --git a/Engine/Source/Runtime/Online/HTTP/Private/GenericPlatform/HttpRequestImpl.cpp b/Engine/Source/Runtime/Online/HTTP/Private/GenericPlatform/HttpRequestImpl.cpp
--- a/Engine/Source/Runtime/Online/HTTP/Private/GenericPlatform/HttpRequestImpl.cpp
+++ b/Engine/Source/Runtime/Online/HTTP/Private/GenericPlatform/HttpRequestImpl.cpp
@@ -36,6 +36,17 @@ FHttpRequestWillRetryDelegate& FHttpRequestImpl::OnRequestWillRetry()
 	return OnRequestWillRetryDelegate;
 }
 
+void FHttpRequestImpl::Shutdown()
+{
+	OnProcessRequestComplete().Unbind();
+PRAGMA_DISABLE_DEPRECATION_WARNINGS
+	OnRequestProgress().Unbind();
+PRAGMA_ENABLE_DEPRECATION_WARNINGS
+	OnRequestProgress64().Unbind();
+	OnStatusCodeReceived().Unbind();
+	OnHeaderReceived().Unbind();
+}
+
 void FHttpRequestImpl::BroadcastResponseHeadersReceived()
 {
 	if (OnHeaderReceived().IsBound())
diff --git a/Engine/Source/Runtime/Online/HTTP/Private/HttpManager.cpp b/Engine/Source/Runtime/Online/HTTP/Private/HttpManager.cpp
--- a/Engine/Source/Runtime/Online/HTTP/Private/HttpManager.cpp
+++ b/Engine/Source/Runtime/Online/HTTP/Private/HttpManager.cpp
@@ -83,14 +83,8 @@ void FHttpManager::Shutdown()
 	// Clear delegates since they may point to deleted instances
 	for (TArray<FHttpRequestRef>::TIterator It(Requests); It; ++It)
 	{
-		FHttpRequestRef& Request = *It;
-		Request->OnProcessRequestComplete().Unbind();
-PRAGMA_DISABLE_DEPRECATION_WARNINGS
-		Request->OnRequestProgress().Unbind();
-PRAGMA_ENABLE_DEPRECATION_WARNINGS
-		Request->OnRequestProgress64().Unbind();
-		Request->OnStatusCodeReceived().Unbind();
-		Request->OnHeaderReceived().Unbind();
+		TSharedPtr<IHttpRequest> Request = *It;
+		StaticCastSharedPtr<FHttpRequestImpl>(Request)->Shutdown();
 
 		// Don't emit these tracking logs in commandlet runs. Build system traps warnings during cook, and these are not truly fatal, but useful for tracking down shutdown issues.
 		UE_CLOG(ShouldOutputHttpWarnings(), LogHttp, Warning, TEXT("	verb=[%s] url=[%s] refs=[%d] status=%s"), *Request->GetVerb(), *Request->GetURL(), Request.GetSharedReferenceCount(), EHttpRequestStatus::ToString(Request->GetStatus()));
@@ -285,10 +279,16 @@ void FHttpManager::AddGameThreadTask(TFunction<void()>&& Task)
 	}
 }
 
-void FHttpManager::AddHttpThreadTask(TFunction<void()>&& Task, float InDelay)
+TSharedPtr<IHttpTaskTimerHandle> FHttpManager::AddHttpThreadTask(TFunction<void()>&& Task, float InDelay)
 {
 	check(Thread);
-	Thread->AddHttpThreadTask(MoveTemp(Task), InDelay);
+	return Thread->AddHttpThreadTask(MoveTemp(Task), InDelay);
+}
+
+void FHttpManager::RemoveHttpThreadTask(TSharedPtr<IHttpTaskTimerHandle> HttpTaskTimerHandle)
+{
+	check(Thread);
+	HttpTaskTimerHandle->RemoveTaskFrom(Thread);
 }
 
 FHttpThreadBase* FHttpManager::CreateHttpThread()
@@ -406,17 +406,25 @@ bool FHttpManager::Tick(float DeltaSeconds)
 {
     QUICK_SCOPE_CYCLE_COUNTER(STAT_FHttpManager_Tick);
 
+	// Normally Tick() should only be called from game thread. But it's still possible Tick() be called 
+	// from off-game thread when quit in purpose like GPU OOM, to flush remain HTTP analysis requests
+
 	// Run GameThread tasks
-	TFunction<void()> Task = nullptr;
-	while (GameThreadQueue.Dequeue(Task))
 	{
-		check(Task);
-		Task();
+		FScopeLock ScopeLock(&GameThreadQueueLock);
+
+		TFunction<void()> Task = nullptr;
+		while (GameThreadQueue.Dequeue(Task))
+		{
+			check(Task);
+			Task();
+		}
 	}
 
+	FScopeLock ScopeLock(&RequestLock);
+
 	if (Thread)
 	{
-		FScopeLock ScopeLock(&RequestLock);
 		// Tick each active request
 		for (const FHttpRequestRef& Request: Requests)
 		{
@@ -441,8 +449,7 @@ bool FHttpManager::Tick(float DeltaSeconds)
 	else
 	{
 		TArray<FHttpRequestRef> CompletedRequests;
-		
-		FScopeLock ScopeLock(&RequestLock);
+
 		// Tick each active request
 		for (const FHttpRequestRef& Request: Requests)
 		{
@@ -489,7 +496,9 @@ void FHttpManager::AddThreadedRequest(const TSharedRef<IHttpThreadedRequest, ESP
 {
 	check(Thread);
 	{
+PRAGMA_DISABLE_DEPRECATION_WARNINGS
 		AddRequest(Request);
+PRAGMA_ENABLE_DEPRECATION_WARNINGS
 	}
 	Thread->AddRequest(&Request.Get());
 }
diff --git a/Engine/Source/Runtime/Online/HTTP/Private/HttpModule.cpp b/Engine/Source/Runtime/Online/HTTP/Private/HttpModule.cpp
--- a/Engine/Source/Runtime/Online/HTTP/Private/HttpModule.cpp
+++ b/Engine/Source/Runtime/Online/HTTP/Private/HttpModule.cpp
@@ -44,8 +44,10 @@ static bool ShouldLaunchUrl(const TCHAR* Url)
 
 void FHttpModule::UpdateConfigs()
 {
-	GConfig->GetFloat(TEXT("HTTP"), TEXT("HttpTimeout"), HttpTimeout, GEngineIni);
+	GConfig->GetFloat(TEXT("HTTP"), TEXT("HttpTimeout"), HttpActivityTimeout, GEngineIni);
+	GConfig->GetFloat(TEXT("HTTP"), TEXT("HttpTotalTimeout"), HttpTotalTimeout, GEngineIni);
 	GConfig->GetFloat(TEXT("HTTP"), TEXT("HttpConnectionTimeout"), HttpConnectionTimeout, GEngineIni);
+	GConfig->GetFloat(TEXT("HTTP"), TEXT("HttpActivityTimeout"), HttpActivityTimeout, GEngineIni);
 	GConfig->GetFloat(TEXT("HTTP"), TEXT("HttpReceiveTimeout"), HttpReceiveTimeout, GEngineIni);
 	GConfig->GetFloat(TEXT("HTTP"), TEXT("HttpSendTimeout"), HttpSendTimeout, GEngineIni);
 	GConfig->GetInt(TEXT("HTTP"), TEXT("HttpMaxConnectionsPerServer"), HttpMaxConnectionsPerServer, GEngineIni);
@@ -77,8 +79,9 @@ void FHttpModule::StartupModule()
 	Singleton = this;
 
 	MaxReadBufferSize = 256 * 1024;
-	HttpTimeout = 300.0f;
+	HttpTotalTimeout = 0.0f;
 	HttpConnectionTimeout = 60.0f;
+	HttpActivityTimeout = 300.0f;
 	HttpReceiveTimeout = HttpConnectionTimeout;
 	HttpSendTimeout = HttpConnectionTimeout;
 	HttpMaxConnectionsPerServer = 16;
diff --git a/Engine/Source/Runtime/Online/HTTP/Private/HttpRetrySystem.cpp b/Engine/Source/Runtime/Online/HTTP/Private/HttpRetrySystem.cpp
--- a/Engine/Source/Runtime/Online/HTTP/Private/HttpRetrySystem.cpp
+++ b/Engine/Source/Runtime/Online/HTTP/Private/HttpRetrySystem.cpp
@@ -16,7 +16,7 @@
 
 TAutoConsoleVariable<bool> CVarHttpRetrySystemNonGameThreadSupportEnabled(
 	TEXT("Http.RetrySystemNonGameThreadSupportEnabled"),
-	false,
+	true,
 	TEXT("Enable retry system non-game thread support")
 );
 
@@ -736,6 +736,19 @@ bool FHttpRetrySystem::FManager::ProcessRequest(TSharedRef<FHttpRetrySystem::FRe
 {
 	QUICK_SCOPE_CYCLE_COUNTER(STAT_FHttpRetrySystem_FManager_ProcessRequest);
 
+	if (CVarHttpRetrySystemNonGameThreadSupportEnabled.GetValueOnAnyThread())
+	{
+		// Let the request trigger timeout by itself, instead of ticking it in retry system
+		if (HttpRetryRequest->RetryTimeoutRelativeSecondsOverride.IsSet())
+		{
+			HttpRetryRequest->SetTimeout(HttpRetryRequest->RetryTimeoutRelativeSecondsOverride.GetValue());
+		}
+		else if (RetryTimeoutRelativeSecondsDefault.IsSet())
+		{
+			HttpRetryRequest->SetTimeout(RetryTimeoutRelativeSecondsDefault.GetValue());
+		}
+	}
+
 	FScopeLock ScopeLock(&RequestListLock);
 	RequestList.Add(FHttpRetryRequestEntry(HttpRetryRequest));
 	HttpRetryRequest->RetryStatus = FHttpRetrySystem::FRequest::EStatus::Processing;
diff --git a/Engine/Source/Runtime/Online/HTTP/Private/HttpThread.cpp b/Engine/Source/Runtime/Online/HTTP/Private/HttpThread.cpp
--- a/Engine/Source/Runtime/Online/HTTP/Private/HttpThread.cpp
+++ b/Engine/Source/Runtime/Online/HTTP/Private/HttpThread.cpp
@@ -25,6 +25,23 @@ DECLARE_CYCLE_STAT(TEXT("CompleteThreadedRequest"), STAT_HTTPThread_CompleteThre
 DECLARE_CYCLE_STAT(TEXT("ActiveSleep"), STAT_HTTPThread_ActiveSleep, STATGROUP_HTTPThread);
 DECLARE_CYCLE_STAT(TEXT("IdleSleep"), STAT_HTTPThread_IdleSleep, STATGROUP_HTTPThread);
 
+class FHttpTaskTimerHandleFTSTicker : public IHttpTaskTimerHandle
+{
+public:
+	FHttpTaskTimerHandleFTSTicker(FTSTicker::FDelegateHandle InHandle)
+		: Handle(InHandle)
+	{
+	}
+
+	virtual void RemoveTaskFrom(FHttpThreadBase* HttpThreadBase) 
+	{ 
+		HttpThreadBase->RemoveTimerHandle(Handle);
+	}
+
+private:
+	FTSTicker::FDelegateHandle Handle;
+};
+
 // FHttpThread
 
 FHttpThreadBase::FHttpThreadBase()
@@ -352,12 +369,22 @@ void FLegacyHttpThread::Tick()
 	}
 }
 
-void FLegacyHttpThread::AddHttpThreadTask(TFunction<void()>&& Task, float InDelay)
+TSharedPtr<IHttpTaskTimerHandle> FLegacyHttpThread::AddHttpThreadTask(TFunction<void()>&& Task, float InDelay)
 {
-	Ticker.AddTicker(FTickerDelegate::CreateLambda([this, Task=MoveTemp(Task)](float) {
+	return MakeShared<FHttpTaskTimerHandleFTSTicker>(Ticker.AddTicker(FTickerDelegate::CreateLambda([this, Task=MoveTemp(Task)](float) {
 		Task();
 		return false;
-	}), InDelay);
+	}), InDelay));
+}
+
+void FLegacyHttpThread::RemoveTimerHandle(FTSTicker::FDelegateHandle DelegateHandle)
+{
+	Ticker.RemoveTicker(DelegateHandle);
+}
+
+void FLegacyHttpThread::RemoveTimerHandle(UE::EventLoop::FTimerHandle EventLoopTimerHandle)
+{
+	checkNoEntry();
 }
 
 void FLegacyHttpThread::HttpThreadTick(float DeltaSeconds)
diff --git a/Engine/Source/Runtime/Online/HTTP/Private/HttpThread.h b/Engine/Source/Runtime/Online/HTTP/Private/HttpThread.h
--- a/Engine/Source/Runtime/Online/HTTP/Private/HttpThread.h
+++ b/Engine/Source/Runtime/Online/HTTP/Private/HttpThread.h
@@ -3,6 +3,7 @@
 #pragma once
 
 #include "CoreMinimal.h"
+#include "EventLoop/EventLoopTimer.h"
 #include "HAL/ThreadSafeCounter.h"
 #include "HAL/Runnable.h"
 #include "HttpPackage.h"
@@ -15,6 +16,15 @@
 #include <atomic>
 
 class IHttpThreadedRequest;
+class FHttpThreadBase;
+
+class IHttpTaskTimerHandle
+{
+public:
+	virtual ~IHttpTaskTimerHandle() {};
+
+	virtual void RemoveTaskFrom(FHttpThreadBase* HttpThreadBase) = 0;
+};
 
 /**
  * Manages Http thread
@@ -84,8 +94,13 @@ public:
 	 *
 	 * @param Task The task to be ran
 	 * @param InDelay The delay to wait before running the task
+	 * @return The handle of the timer, which could be used to remove the task before it get triggered
 	 */
-	virtual void AddHttpThreadTask(TFunction<void()>&& Task, float InDelay) = 0;
+	virtual TSharedPtr<IHttpTaskTimerHandle> AddHttpThreadTask(TFunction<void()>&& Task, float InDelay) = 0;
+
+	virtual void RemoveTimerHandle(FTSTicker::FDelegateHandle DelegateHandle) = 0;
+
+	virtual void RemoveTimerHandle(UE::EventLoop::FTimerHandle EventLoopTimerHandle) = 0;
 
 protected:
 
@@ -199,9 +214,12 @@ protected:
 	virtual void Stop() override;
 	//~ End FRunnable Interface
 
-	virtual void AddHttpThreadTask(TFunction<void()>&& Task, float InDelay) override;
+	virtual TSharedPtr<IHttpTaskTimerHandle> AddHttpThreadTask(TFunction<void()>&& Task, float InDelay) override;
 	virtual void HttpThreadTick(float DeltaSeconds) override;
 
+	virtual void RemoveTimerHandle(FTSTicker::FDelegateHandle DelegateHandle) override;
+	virtual void RemoveTimerHandle(UE::EventLoop::FTimerHandle EventLoopTimerHandle) override;
+
 	/** signal request to stop and exit thread */
 	FThreadSafeCounter ExitRequest;
 
diff --git a/Engine/Source/Runtime/Online/HTTP/Private/NullHttp.cpp b/Engine/Source/Runtime/Online/HTTP/Private/NullHttp.cpp
--- a/Engine/Source/Runtime/Online/HTTP/Private/NullHttp.cpp
+++ b/Engine/Source/Runtime/Online/HTTP/Private/NullHttp.cpp
@@ -130,8 +130,6 @@ bool FNullHttpRequest::ProcessRequest()
 	CompletionStatus = EHttpRequestStatus::Processing;
 
 	UE_LOG(LogHttp, Log, TEXT("Start request. %p %s url=%s"), this, *GetVerb(), *GetURL());
-
-	FHttpModule::Get().GetHttpManager().AddRequest(SharedThis(this));
 	return true;
 }
 
@@ -170,7 +168,7 @@ void FNullHttpRequest::Tick(float DeltaSeconds)
 	if (CompletionStatus == EHttpRequestStatus::Processing)
 	{
 		ElapsedTime += DeltaSeconds;
-		const float HttpTimeout = GetTimeout().Get(FHttpModule::Get().GetHttpTimeout());
+		const float HttpTimeout = GetTimeout().Get(FHttpModule::Get().GetHttpActivityTimeout());
 		if (HttpTimeout > 0 && ElapsedTime >= HttpTimeout)
 		{
 			UE_LOG(LogHttp, Warning, TEXT("Timeout processing Http request. %p"),
diff --git a/Engine/Source/Runtime/Online/HTTP/Private/WinHttp/Support/WinHttpSession.cpp b/Engine/Source/Runtime/Online/HTTP/Private/WinHttp/Support/WinHttpSession.cpp
--- a/Engine/Source/Runtime/Online/HTTP/Private/WinHttp/Support/WinHttpSession.cpp
+++ b/Engine/Source/Runtime/Online/HTTP/Private/WinHttp/Support/WinHttpSession.cpp
@@ -85,8 +85,8 @@ FWinHttpSession::FWinHttpSession(uint32 SecurityProtocolFlags, const bool bInFor
 	FHttpModule& HttpModule = FHttpModule::Get();
 	const FTimespan ConnectionTimeout = FTimespan::FromSeconds(HttpModule.GetHttpConnectionTimeout() > 0 ? HttpModule.GetHttpConnectionTimeout() : 0);
 	const FTimespan ResolveTimeout = ConnectionTimeout;
-	const FTimespan SendTimeout = FTimespan::FromSeconds(HttpModule.GetHttpSendTimeout() > 0 ? HttpModule.GetHttpSendTimeout() : 0);
-	const FTimespan ReceiveTimeout = FTimespan::FromSeconds(HttpModule.GetHttpReceiveTimeout() > 0 ? HttpModule.GetHttpReceiveTimeout() : 0);
+	const FTimespan ReceiveTimeout = FTimespan::FromSeconds(HttpModule.GetHttpActivityTimeout() > 0 ? HttpModule.GetHttpActivityTimeout() : 0);
+	const FTimespan SendTimeout = ReceiveTimeout;
 
 	if (!WinHttpSetTimeouts(SessionHandle.Get(), ResolveTimeout.GetTotalMilliseconds(), ConnectionTimeout.GetTotalMilliseconds(), SendTimeout.GetTotalMilliseconds(), ReceiveTimeout.GetTotalMilliseconds()))
 	{
diff --git a/Engine/Source/Runtime/Online/HTTP/Public/GenericPlatform/HttpRequestCommon.h b/Engine/Source/Runtime/Online/HTTP/Public/GenericPlatform/HttpRequestCommon.h
--- a/Engine/Source/Runtime/Online/HTTP/Public/GenericPlatform/HttpRequestCommon.h
+++ b/Engine/Source/Runtime/Online/HTTP/Public/GenericPlatform/HttpRequestCommon.h
@@ -4,12 +4,16 @@
 
 #include "GenericPlatform/HttpRequestImpl.h"
 
+class IHttpTaskTimerHandle;
+
 /**
  * Contains implementation of some common functions that don't vary between implementations of different platforms
  */
 class FHttpRequestCommon : public FHttpRequestImpl
 {
 public:
+	FHttpRequestCommon();
+
 	// IHttpBase
 	HTTP_API virtual FString GetURLParameter(const FString& ParameterName) const override;
 
@@ -27,14 +31,19 @@ public:
 	// Can be called on game thread or http thread depend on the delegate thread policy
 	HTTP_API virtual void FinishRequest() = 0;
 
+	HTTP_API virtual void CancelRequest() override;
+
+	HTTP_API virtual void Shutdown() override;
+
 protected:
 	/**
 	 * Check if this request is valid or allowed, before actually process the request
 	 */
 	HTTP_API bool PreProcess();
+	HTTP_API void PostProcess();
 	HTTP_API virtual bool SetupRequest() = 0;
 	HTTP_API bool PreCheck() const;
-	HTTP_API void ClearInCaseOfRetry();
+	HTTP_API virtual void ClearInCaseOfRetry();
 
 	HTTP_API void SetStatus(EHttpRequestStatus::Type InCompletionStatus);
 	HTTP_API void SetFailureReason(EHttpFailureReason InFailureReason);
@@ -46,6 +55,19 @@ protected:
 
 	HTTP_API void HandleRequestSucceed(TSharedPtr<IHttpResponse> Response);
 
+	HTTP_API void StartActivityTimeoutTimer();
+	HTTP_API void StartActivityTimeoutTimerBy(double DelayToTrigger);
+	HTTP_API void ResetActivityTimeoutTimer(FStringView Reason);
+	HTTP_API void OnActivityTimeoutTimerTaskTrigger();
+	HTTP_API void StopActivityTimeoutTimer();
+	HTTP_API void StartTotalTimeoutTimer();
+	HTTP_API void StopTotalTimeoutTimer();
+	HTTP_API void OnTotalTimeoutTimerTaskTrigger();
+
+	HTTP_API virtual void AbortRequest() = 0;
+
+	HTTP_API virtual void CleanupRequest() = 0;
+
 	HTTP_API void TriggerStatusCodeReceivedDelegate(int32 StatusCode);
 
 protected:
@@ -61,9 +83,29 @@ protected:
 	/** Timeout in seconds for the entire HTTP request to complete */
 	TOptional<float> TimeoutSecs;
 
+	/** Indicate the request is timed out, it should quit and fail with EHttpFailureReason::TimedOut */
+	bool bTimedOut = false;
+	/** Indicate the request is activity timed out, it should quit and fail with EHttpFailureReason::ConnectionError */
+	bool bActivityTimedOut = false;
+	/** Indicate the request is cancelled, it should quit and fail with EHttpFailureReason::Cancelled */
+	bool bCanceled = false;
+
+	/** TODO: Move this feature into CurlHttp */
+	bool bUsePlatformActivityTimeout = true;
+
 	/** Record when this request started */
 	double RequestStartTimeAbsoluteSeconds;
 
+	/** Record when this request will activity timeout */
+	double ActivityTimeoutAt;
+
+	/** Holder the timer handle, if the request get destroyed before triggering the timeout, use this to remove the timer */
+	TSharedPtr<IHttpTaskTimerHandle> TotalTimeoutHttpTaskTimerHandle;
+	/** Holder the timer handle, if the request get destroyed before triggering the timeout, use this to remove the timer */
+	TSharedPtr<IHttpTaskTimerHandle> ActivityTimeoutHttpTaskTimerHandle;
+	/** Critical section for accessing HttpTaskTimerHandle */
+	FCriticalSection HttpTaskTimerHandleCriticalSection;
+
 	/** Record when the request start to process */
 	double StartProcessTime = 0.0;
 
diff --git a/Engine/Source/Runtime/Online/HTTP/Public/GenericPlatform/HttpRequestImpl.h b/Engine/Source/Runtime/Online/HTTP/Public/GenericPlatform/HttpRequestImpl.h
--- a/Engine/Source/Runtime/Online/HTTP/Public/GenericPlatform/HttpRequestImpl.h
+++ b/Engine/Source/Runtime/Online/HTTP/Public/GenericPlatform/HttpRequestImpl.h
@@ -19,6 +19,11 @@ public:
 	HTTP_API virtual FHttpRequestHeaderReceivedDelegate& OnHeaderReceived() override;
 	HTTP_API virtual FHttpRequestWillRetryDelegate& OnRequestWillRetry() override;
 
+	/**
+	 * Clear callbacks before http module get deleted
+	 */
+	HTTP_API virtual void Shutdown();
+
 protected:
 	/** 
 	 * Broadcast all of our response's headers as having been received
diff --git a/Engine/Source/Runtime/Online/HTTP/Public/HttpManager.h b/Engine/Source/Runtime/Online/HTTP/Public/HttpManager.h
--- a/Engine/Source/Runtime/Online/HTTP/Public/HttpManager.h
+++ b/Engine/Source/Runtime/Online/HTTP/Public/HttpManager.h
@@ -20,6 +20,7 @@
 
 class FHttpThreadBase;
 class FOutputDevice;
+class IHttpTaskTimerHandle;
 class IHttpThreadedRequest;
 
 enum class EHttpFlushReason : uint8
@@ -98,6 +99,7 @@ public:
 	 *
 	 * @param Request - the request object to add
 	 */
+	UE_DEPRECATED(5.4, "AddRequest has been deprecated, use AddThreadedRequest instead")
 	HTTP_API void AddRequest(const FHttpRequestRef& Request);
 
 	/**
@@ -243,12 +245,20 @@ public:
 	HTTP_API void AddGameThreadTask(TFunction<void()>&& Task);
 
 	/**
-	 * Add task to be ran on the http thread next tick
+	 * Add task to be ran on the http thread
 	 *
 	 * @param Task The task to be ran
 	 * @param InDelay The delay to wait before running the task
+	 * @return The handle of the timer, which could be used to remove the task before it's triggered
 	 */
-	HTTP_API void AddHttpThreadTask(TFunction<void()>&& Task, float InDelay = 0.0f);
+	HTTP_API TSharedPtr<IHttpTaskTimerHandle> AddHttpThreadTask(TFunction<void()>&& Task, float InDelay = 0.0f);
+
+	/**
+	 * Remove the task from the http thread before it's triggered
+	 *
+	 * @param HttpTaskTimerHandle The handle of the timer
+	 */
+	HTTP_API void RemoveHttpThreadTask(TSharedPtr<IHttpTaskTimerHandle> HttpTaskTimerHandle);
 
 	/**
 	 * Set url request filter through code, instead of setting it through config.
@@ -282,6 +292,7 @@ protected:
 
 	/** Queue of tasks to run on the game thread */
 	TQueue<TFunction<void()>, EQueueMode::Mpsc> GameThreadQueue;
+	FCriticalSection GameThreadQueueLock;
 
 	// This variable is set to true in Flush(EHttpFlushReason), and prevents new Http requests from being launched
 	bool bFlushing;
diff --git a/Engine/Source/Runtime/Online/HTTP/Public/HttpModule.h b/Engine/Source/Runtime/Online/HTTP/Public/HttpModule.h
--- a/Engine/Source/Runtime/Online/HTTP/Public/HttpModule.h
+++ b/Engine/Source/Runtime/Online/HTTP/Public/HttpModule.h
@@ -84,17 +84,27 @@ public:
 	/**
 	 * @return timeout in seconds for the entire http request to complete
 	 */
+	UE_DEPRECATED(5.4, "GetHttpTimeout has been deprecated, use GetHttpActivityTimeout or GetHttpTotalTimeout instead")
 	inline float GetHttpTimeout() const
 	{
-		return HttpTimeout;
+		return HttpTotalTimeout;
+	}
+
+	/**
+	 * @return total timeout in seconds for the entire http request to complete
+	 */
+	inline float GetHttpTotalTimeout() const
+	{
+		return HttpTotalTimeout;
 	}
 
 	/**
 	 * Sets timeout in seconds for the entire http request to complete
 	 */
+	UE_DEPRECATED(5.4, "SetHttpTimeout has been deprecated, config it through HttpActivityTimeout or HttpTotalTimeout instead")
 	inline void SetHttpTimeout(float TimeOutInSec)
 	{
-		HttpTimeout = TimeOutInSec;
+		HttpTotalTimeout = TimeOutInSec;
 	}
 
 	/**
@@ -108,6 +118,7 @@ public:
 	/**
 	 * @return timeout in seconds to receive a response on the connection 
 	 */
+	UE_DEPRECATED(5.4, "GetHttpReceiveTimeout has been deprecated, Use GetHttpActivityTimeout instead to decide if there is still any ongoing activity.")
 	inline float GetHttpReceiveTimeout() const
 	{
 		return HttpReceiveTimeout;
@@ -116,11 +127,20 @@ public:
 	/**
 	 * @return timeout in seconds to send a request on the connection
 	 */
+	UE_DEPRECATED(5.4, "GetHttpSendTimeout has been deprecated, The legacy behavior doesn't make use of HttpSendTimeout at all, and there is no such support on some platforms. Use GetHttpActivityTimeout instead to decide if there is still any ongoing activity.")
 	inline float GetHttpSendTimeout() const
 	{
 		return HttpSendTimeout;
 	}
 
+	/**
+	 * @return timeout in seconds to check there is any ongoing activity on the established connection
+	 */
+	inline float GetHttpActivityTimeout() const
+	{
+		return HttpActivityTimeout;
+	}
+
 	/**
 	 * @return max number of simultaneous connections to a specific server
 	 */
@@ -318,6 +338,10 @@ public:
 protected:
 	/** timeout in seconds to establish the connection */
 	float HttpConnectionTimeout;
+	/** timeout in seconds for the entire http request to complete. 0 is no timeout */
+	float HttpTotalTimeout;
+	/**  timeout in seconds to check there is any ongoing activity on the established connection */
+	float HttpActivityTimeout;
 
 private:
 
@@ -353,8 +377,6 @@ private:
 
 	/** Keeps track of Http requests while they are being processed */
 	FHttpManager* HttpManager = nullptr;
-	/** timeout in seconds for the entire http request to complete. 0 is no timeout */
-	float HttpTimeout;
 	/** timeout in seconds to receive a response on the connection */
 	float HttpReceiveTimeout;
 	/** timeout in seconds to send a request on the connection */
