Add check in WebTests test case to validate the received data.

diff --git a/Engine/Source/Programs/WebTests/Private/TestHttp.cpp b/Engine/Source/Programs/WebTests/Private/TestHttp.cpp
--- a/Engine/Source/Programs/WebTests/Private/TestHttp.cpp
+++ b/Engine/Source/Programs/WebTests/Private/TestHttp.cpp
@@ -299,12 +299,37 @@ TEST_CASE_METHOD(FWaitUntilCompleteHttpFixture, "Can do blocking call", HTTP_TAG
 TEST_CASE_METHOD(FWaitUntilCompleteHttpFixture, "Get large response content without chunks", HTTP_TAG)
 {
 	TSharedRef<IHttpRequest> HttpRequest = CreateRequest();
-	HttpRequest->SetURL(FString::Format(TEXT("{0}/get_large_response_without_chunks/{1}/"), { *UrlHttpTests(), 1024 * 1024/*bytes_number*/}));
+	uint32 DataLength = 0;
+	uint32 RepeatAt = 0;
+	SECTION("case A")
+	{
+		DataLength = 1024 * 1024;
+		RepeatAt = 10;
+	}
+	SECTION("cast B")
+	{
+		DataLength = 1025 * 1023;
+		RepeatAt = 9;
+	}
+	HttpRequest->SetURL(FString::Format(TEXT("{0}/get_data_without_chunks/{1}/{2}/"), { *UrlHttpTests(), DataLength, RepeatAt}));
 	HttpRequest->SetVerb(TEXT("GET"));
-	HttpRequest->OnProcessRequestComplete().BindLambda([](FHttpRequestPtr HttpRequest, FHttpResponsePtr HttpResponse, bool bSucceeded) {
+	HttpRequest->OnProcessRequestComplete().BindLambda([DataLength, RepeatAt](FHttpRequestPtr HttpRequest, FHttpResponsePtr HttpResponse, bool bSucceeded) {
 		CHECK(bSucceeded);
 		REQUIRE(HttpResponse != nullptr);
 		CHECK(HttpResponse->GetResponseCode() == 200);
+
+		const TArray<uint8>& Content = HttpResponse->GetContent();
+		CHECK(Content.Num() == DataLength);
+
+		bool bAllMatch = true;
+
+		// Make sure the data got read is in good state
+		for (int32 i = 0; i < Content.Num(); ++i)
+		{
+			bAllMatch &= (Content[i] == '0' + (i % RepeatAt));
+		}
+
+		CHECK(bAllMatch);
 	});
 	HttpRequest->ProcessRequest();
 }
diff --git a/Engine/Source/Programs/WebTestsServer/httptests/urls.py b/Engine/Source/Programs/WebTestsServer/httptests/urls.py
--- a/Engine/Source/Programs/WebTestsServer/httptests/urls.py
+++ b/Engine/Source/Programs/WebTestsServer/httptests/urls.py
@@ -5,7 +5,7 @@ from . import views
 urlpatterns = [
     path("httptests/methods", views.methods, name="methods"),
     path("httptests/query_with_params/", views.query_with_params, name="query_with_params"),
-    path("httptests/get_large_response_without_chunks/<int:bytes_number>/", views.get_large_response_without_chunks, name="get_large_response_without_chunks"),
+    path("httptests/get_data_without_chunks/<int:bytes_number>/<int:repeat_at>/", views.get_data_without_chunks, name="get_data_without_chunks"),
     path("httptests/streaming_download/<int:chunks>/<int:chunk_size>/<int:chunk_latency>/", views.streaming_download, name="streaming_download"),
     path("httptests/streaming_upload_post", views.streaming_upload_post, name="streaming_upload_post"),
     path("httptests/streaming_upload_put", views.streaming_upload_put, name="streaming_upload_put"),
diff --git a/Engine/Source/Programs/WebTestsServer/httptests/views.py b/Engine/Source/Programs/WebTestsServer/httptests/views.py
--- a/Engine/Source/Programs/WebTestsServer/httptests/views.py
+++ b/Engine/Source/Programs/WebTestsServer/httptests/views.py
@@ -3,6 +3,7 @@ from datetime import datetime
 from django.shortcuts import render, redirect
 from django.http.response import JsonResponse
 from django.http import JsonResponse
+from django.http import HttpResponse
 from django.http import StreamingHttpResponse
 from rest_framework.decorators import api_view, parser_classes
 from rest_framework.parsers import FileUploadParser
@@ -22,9 +23,11 @@ def query_with_params(request):
     return JsonResponse({'var_int' : var_int, 'var_str' : var_str}, content_type="application/json")
 
 @api_view(['GET'])
-def get_large_response_without_chunks(request, bytes_number):
-    data = "d" * bytes_number
-    return JsonResponse({'data' : data})
+def get_data_without_chunks(request, bytes_number, repeat_at):
+    assert repeat_at > 0
+    assert repeat_at <= 10
+    data = [i%repeat_at for i in range(bytes_number)]
+    return HttpResponse(data)
 
 @api_view(['GET'])
 def streaming_download(request, chunks, chunk_size, chunk_latency):
