Fix the failure in test case "Can do blocking call" in WebTests by removing the complete delegate.

diff --git a/Engine/Source/Programs/WebTests/Private/TestHttp.cpp b/Engine/Source/Programs/WebTests/Private/TestHttp.cpp
--- a/Engine/Source/Programs/WebTests/Private/TestHttp.cpp
+++ b/Engine/Source/Programs/WebTests/Private/TestHttp.cpp
@@ -288,10 +288,6 @@ TEST_CASE_METHOD(FWaitUntilCompleteHttpFixture, "Can do blocking call", HTTP_TAG
 {
 	TSharedRef<IHttpRequest> HttpRequest = CreateRequest();
 	HttpRequest->SetURL(UrlToTestMethods());
-	HttpRequest->OnProcessRequestComplete().BindLambda([](FHttpRequestPtr HttpRequest, FHttpResponsePtr HttpResponse, bool bSucceeded) {
-		CHECK(bSucceeded);
-		REQUIRE(HttpResponse != nullptr);
-	});
 	HttpRequest->ProcessRequestUntilComplete();
 	CHECK(HttpRequest->GetStatus() == EHttpRequestStatus::Succeeded);
 }
