Make sure python packages can be updated in WebTestsServer;
Added more checks in test case "Can do blocking call" of WebTests;
Added duplicated call CancelRequest in test case "Cancel http request connect before timeout" of  WebTests;
Some other code clean in WebTests.

diff --git a/Engine/Source/Programs/WebTests/Private/TestHttp.cpp b/Engine/Source/Programs/WebTests/Private/TestHttp.cpp
--- a/Engine/Source/Programs/WebTests/Private/TestHttp.cpp
+++ b/Engine/Source/Programs/WebTests/Private/TestHttp.cpp
@@ -290,6 +290,9 @@ TEST_CASE_METHOD(FWaitUntilCompleteHttpFixture, "Can do blocking call", HTTP_TAG
 	HttpRequest->SetURL(UrlToTestMethods());
 	HttpRequest->ProcessRequestUntilComplete();
 	CHECK(HttpRequest->GetStatus() == EHttpRequestStatus::Succeeded);
+	FHttpResponsePtr HttpResponse = HttpRequest->GetResponse();
+	REQUIRE(HttpResponse != nullptr);
+	CHECK(HttpResponse->GetResponseCode() == 200);
 }
 
 TEST_CASE_METHOD(FWaitUntilCompleteHttpFixture, "Get large response content without chunks", HTTP_TAG)
@@ -1272,10 +1275,10 @@ TEST_CASE_METHOD(FWaitUntilQuitFromTestFixture, "Http request pre check will fai
 	HttpRequest->ProcessRequest();
 }
 
-class FWaitUnitilQuitFromTestThreadedFixture : public FWaitUntilQuitFromTestFixture
+class FWaitUntilQuitFromTestThreadedFixture : public FWaitUntilQuitFromTestFixture
 {
 public:
-	~FWaitUnitilQuitFromTestThreadedFixture()
+	~FWaitUntilQuitFromTestThreadedFixture()
 	{
 		WaitUntilQuitFromTest();
 	}
@@ -1284,7 +1287,7 @@ public:
 };
 
 // Pre-check failed requests won't be added into http manager, so it can't rely on the requested added/completed callback in FWaitUntilCompleteHttpFixture
-TEST_CASE_METHOD(FWaitUnitilQuitFromTestThreadedFixture, "Threaded http request pre check will fail by thread policy", HTTP_TAG)
+TEST_CASE_METHOD(FWaitUntilQuitFromTestThreadedFixture, "Threaded http request pre check will fail by thread policy", HTTP_TAG)
 {
 	DisableWarningsInThisTest();
 
@@ -1338,6 +1341,7 @@ TEST_CASE_METHOD(FWaitUntilCompleteHttpFixture, "Cancel http request connect bef
 	HttpRequest->ProcessRequest();
 	FPlatformProcess::Sleep(0.5);
 	HttpRequest->CancelRequest();
+	HttpRequest->CancelRequest(); // Duplicated calls to CancelRequest should be fine
 }
 
 class FThreadedBatchRequestsFixture : public FWaitThreadedHttpFixture
@@ -1367,7 +1371,7 @@ public:
 	}
 };
 
-TEST_CASE_METHOD(FThreadedBatchRequestsFixture, "Retry manager and http manager is thread safe", HTTP_TAG)
+TEST_CASE_METHOD(FThreadedBatchRequestsFixture, "Retry manager and http manager is thread safe for flushing", HTTP_TAG)
 {
 	DisableWarningsInThisTest();
 
@@ -1496,5 +1500,3 @@ TEST_CASE_METHOD(FLocalHttpServerFixture, "Local http server can serve large fil
 
 #endif // (PLATFORM_WINDOWS && !WITH_CURL_XCURL) || PLATFORM_MAC || PLATFORM_UNIX
 
-// TODO: Add cancel test, with multiple cancel calls
-
diff --git a/Engine/Source/Programs/WebTestsServer/createenv.bat b/Engine/Source/Programs/WebTestsServer/createenv.bat
--- a/Engine/Source/Programs/WebTestsServer/createenv.bat
+++ b/Engine/Source/Programs/WebTestsServer/createenv.bat
@@ -1,14 +1,14 @@
 @echo off
 
-if exist env (
-	call .\env\Scripts\activate.bat
-) else (
+if not exist env (
 	..\..\..\Binaries\ThirdParty\Python3\Win64\python.exe -m venv env
-	call .\env\Scripts\activate.bat
-	python -m pip install --upgrade pip
-	pip install -r requirements.txt
-	if %ERRORLEVEL% NEQ 0 (
-		rmdir /s /q env
-		exit %ERRORLEVEL%
-	)
+)
+
+call .\env\Scripts\activate.bat
+
+python -m pip install --upgrade pip
+pip install -r requirements.txt
+if %ERRORLEVEL% NEQ 0 (
+	rmdir /s /q env
+	exit %ERRORLEVEL%
 )
diff --git a/Engine/Source/Programs/WebTestsServer/createenv.sh b/Engine/Source/Programs/WebTestsServer/createenv.sh
--- a/Engine/Source/Programs/WebTestsServer/createenv.sh
+++ b/Engine/Source/Programs/WebTestsServer/createenv.sh
@@ -1,6 +1,4 @@
-if [ -d "env" ]; then
-	. ./env/bin/activate
-else
+if [ ! -d "env" ]; then
 	case "$OSTYPE" in
 		darwin*)
 			../../../Binaries/ThirdParty/Python3/Mac/bin/python3 -m venv env
@@ -9,12 +7,13 @@ else
 			../../../Binaries/ThirdParty/Python3/Linux/bin/python3 -m venv env
 			;;
 	esac
-	. ./env/bin/activate
-	python -m pip install --upgrade pip
-	pip install -r requirements.txt
-	local install_exit_code = $?
-	if [ $install_exit_code -ne 0 ]; then
-		rm -rf env
-		exit $install_exit_code
-	fi
+fi
+
+. ./env/bin/activate
+python -m pip install --upgrade pip
+pip install -r requirements.txt
+local install_exit_code = $?
+if [ $install_exit_code -ne 0 ]; then
+	rm -rf env
+	exit $install_exit_code
 fi
