Adds tests for LAUNCHREQUEST 1, 20, 50, 100, 500, 1000 reqeusts.
Also adds tests for canceling 1,20,50,100,500,1000 requests.

diff --git a/Engine/Source/Programs/WebTests/Private/TestHttp.cpp b/Engine/Source/Programs/WebTests/Private/TestHttp.cpp
--- a/Engine/Source/Programs/WebTests/Private/TestHttp.cpp
+++ b/Engine/Source/Programs/WebTests/Private/TestHttp.cpp
@@ -17,6 +17,8 @@
 #include "TestHarness.h"
 #include "Serialization/JsonSerializerMacros.h"
 
+#include "catch2/generators/catch_generators.hpp"
+
 /**
  *  HTTP Tests
  *  -----------------------------------------------------------------------------------------------
@@ -1690,3 +1692,51 @@ TEST_CASE_METHOD(FLocalHttpServerFixture, "Local http server can serve large fil
 
 #endif // (PLATFORM_WINDOWS && !WITH_CURL_XCURL) || PLATFORM_MAC || PLATFORM_UNIX
 
+TEST_CASE_METHOD(FWaitUntilCompleteHttpFixture, "Test platform request requests limits", HTTP_TAG "[LIMIT]")
+{
+	bool bCheckCancel = GENERATE(false, true);
+	int32 NumRequests = GENERATE(1, 10, 20, 50, 100, 200, 500, 1000);
+	//Output NumRequests when error occurs.
+	UNSCOPED_INFO(NumRequests);
+	UNSCOPED_INFO(bCheckCancel);
+
+	DYNAMIC_SECTION(" making " << NumRequests << " requests with bCheckCancel=" << bCheckCancel)
+	{
+		if (NumRequests > 100 && !bRunHeavyTests)
+		{
+			SKIP("-run_heavy_tests is unset, limiting NumRequests to 100");
+		}
+
+		TArray<TSharedRef<IHttpRequest>> Requests;
+
+		for (int32 i = 0; i < NumRequests; ++i)
+		{
+			TSharedRef<IHttpRequest> HttpRequest = FHttpModule::Get().CreateRequest();
+			// Requests server to serve 1024b chunks to allow time for cancel to happen
+			HttpRequest->SetURL(UrlStreamDownload(2, 1024, 1));
+			HttpRequest->SetVerb(TEXT("GET"));
+			HttpRequest->OnProcessRequestComplete().BindLambda([bCheckCancel](FHttpRequestPtr HttpRequest, FHttpResponsePtr HttpResponse, bool bSucceeded)
+			{
+				//Only assert if response is successful on non-canceled requests
+				if (!bCheckCancel)
+				{
+					CHECK(bSucceeded);
+					CHECK(HttpResponse != nullptr);
+				}
+			});
+			HttpRequest->ProcessRequest();
+
+			Requests.Add(HttpRequest);
+		}
+
+		CHECK(Requests.Num() == NumRequests);
+
+		if(bCheckCancel)
+		{
+			for (auto Request : Requests)
+			{
+				Request->CancelRequest();
+			}
+		}
+	}
+}
\ No newline at end of file
