In http sony impl change it to read response headers once received status code;
Use cvar for the new flow just in case the headers needs to be read when complete for some urls.

#tests Tried through WebTests and the game

diff --git a/Engine/Source/Programs/WebTests/Private/TestHttp.cpp b/Engine/Source/Programs/WebTests/Private/TestHttp.cpp
--- a/Engine/Source/Programs/WebTests/Private/TestHttp.cpp
+++ b/Engine/Source/Programs/WebTests/Private/TestHttp.cpp
@@ -1173,6 +1173,7 @@ public:
 
 	std::atomic<bool> bHeaderReceived = false;
 	std::atomic<bool> bCompleteCallbackTriggered = false;
+	std::atomic<bool> bAnyDataReceived = false;
 };
 
 TEST_CASE_METHOD(FValidateHeaderReceiveOrderFixture, "Http request header received callback will be called by thread policy", HTTP_TAG)
@@ -1181,10 +1182,19 @@ TEST_CASE_METHOD(FValidateHeaderReceiveOrderFixture, "Http request header receiv
 	HttpRequest->SetURL(UrlStreamDownload(2/*Chunks*/, 1024/*ChunkSize*/));
 	HttpRequest->SetVerb(TEXT("GET"));
 
+	FHttpRequestStreamDelegate StreamDelegate;
+	StreamDelegate.BindLambda([this](void *InDataPtr, int64 InLength) {
+		bAnyDataReceived = true;
+		CHECK(!bCompleteCallbackTriggered);
+		return true;
+	});
+	HttpRequest->SetResponseBodyReceiveStreamDelegate(StreamDelegate);
+
 	SECTION("in http thread")
 	{
 		HttpRequest->SetDelegateThreadPolicy(EHttpRequestDelegateThreadPolicy::CompleteOnHttpThread);
 		HttpRequest->OnHeaderReceived().BindLambda([this](FHttpRequestPtr Request, const FString& HeaderName, const FString& HeaderValue) {
+			CHECK(!bAnyDataReceived);
 			CHECK(!bCompleteCallbackTriggered);
 			CHECK(!IsInGameThread());
 			bHeaderReceived = true;
@@ -1194,6 +1204,9 @@ TEST_CASE_METHOD(FValidateHeaderReceiveOrderFixture, "Http request header receiv
 	{
 		HttpRequest->SetDelegateThreadPolicy(EHttpRequestDelegateThreadPolicy::CompleteOnGameThread);
 		HttpRequest->OnHeaderReceived().BindLambda([this](FHttpRequestPtr Request, const FString& HeaderName, const FString& HeaderValue) {
+			// Data received delegate always triggered from http thread, so it could have been received, while header will be received
+			// from game thread in this test section
+			// CHECK(!bAnyDataReceived);
 			CHECK(!bCompleteCallbackTriggered);
 			CHECK(IsInGameThread());
 			bHeaderReceived = true;
@@ -1201,6 +1214,7 @@ TEST_CASE_METHOD(FValidateHeaderReceiveOrderFixture, "Http request header receiv
 	}
 
 	HttpRequest->OnProcessRequestComplete().BindLambda([this](FHttpRequestPtr /*HttpRequest*/, FHttpResponsePtr /*HttpResponse */, bool bSucceeded) {
+		CHECK(bAnyDataReceived);
 		CHECK(bHeaderReceived);
 		bCompleteCallbackTriggered = true;
 		CHECK(bSucceeded);
