Fix the test case "Optionally retry limit can be set differently for connection error", can't really retry when activity timeout, because the response is not null when get connection error. User code might still want to access the response when activity timeout. So changed the test case to use connection time out instead;
Get rid of platform specific code in WebTests, define macros in Http.Build.cs instead.

diff --git a/Engine/Source/Programs/WebTests/Private/TestHttp.cpp b/Engine/Source/Programs/WebTests/Private/TestHttp.cpp
--- a/Engine/Source/Programs/WebTests/Private/TestHttp.cpp
+++ b/Engine/Source/Programs/WebTests/Private/TestHttp.cpp
@@ -34,7 +34,7 @@
  */
 
 #define HTTP_TAG "[HTTP]"
-#define HTTP_TIME_DIFF_TOLERANCE 0.5f
+#define HTTP_TIME_DIFF_TOLERANCE_OF_REQUEST 0.5f
 
 extern TAutoConsoleVariable<bool> CVarHttpInsecureProtocolEnabled;
 extern TAutoConsoleVariable<bool> CVarHttpRetrySystemNonGameThreadSupportEnabled;
@@ -292,6 +292,35 @@ TEST_CASE_METHOD(FWaitUntilCompleteHttpFixture, "Can process https request", HTT
 	HttpRequest->ProcessRequest();
 }
 
+TEST_CASE_METHOD(FWaitUntilCompleteHttpFixture, "Can complete successfully for different response codes", HTTP_TAG)
+{
+	TSharedRef<IHttpRequest> HttpRequest = CreateRequest();
+	HttpRequest->SetVerb(TEXT("GET"));
+
+	int32 ExpectedStatusCode = 0;
+	SECTION("For status 200")
+	{
+		ExpectedStatusCode = 200;
+	}
+	SECTION("For status 206")
+	{
+		ExpectedStatusCode = 206;
+	}
+	SECTION("For status 400")
+	{
+		ExpectedStatusCode = 400;
+	}
+
+	HttpRequest->SetURL(UrlMockStatus(ExpectedStatusCode));
+
+	HttpRequest->OnProcessRequestComplete().BindLambda([ExpectedStatusCode](FHttpRequestPtr HttpRequest, FHttpResponsePtr HttpResponse, bool bSucceeded) {
+		CHECK(bSucceeded);
+		REQUIRE(HttpResponse != nullptr);
+		CHECK(HttpResponse->GetResponseCode() == ExpectedStatusCode);
+	});
+	HttpRequest->ProcessRequest();
+}
+
 TEST_CASE_METHOD(FWaitUntilCompleteHttpFixture, "Can do blocking call", HTTP_TAG)
 {
 	TSharedRef<IHttpRequest> HttpRequest = CreateRequest();
@@ -346,7 +375,8 @@ TEST_CASE_METHOD(FWaitUntilCompleteHttpFixture, "Http request connect timeout",
 	DisableWarningsInThisTest();
 
 	HttpModule->HttpActivityTimeout = 3.0f; // Make sure this won't be triggered before establishing connection
-	HttpModule->HttpConnectionTimeout = 15.0f;
+	float ExpectedTimeoutDuration = 15.0f;
+	HttpModule->HttpConnectionTimeout = ExpectedTimeoutDuration;
 
 	TSharedRef<IHttpRequest> HttpRequest = CreateRequest();
 
@@ -355,17 +385,13 @@ TEST_CASE_METHOD(FWaitUntilCompleteHttpFixture, "Http request connect timeout",
 
 	const double StartTime = FPlatformTime::Seconds();
 
-	HttpRequest->OnProcessRequestComplete().BindLambda([StartTime](FHttpRequestPtr HttpRequest, FHttpResponsePtr HttpResponse, bool bSucceeded) {
+	HttpRequest->OnProcessRequestComplete().BindLambda([StartTime, ExpectedTimeoutDuration](FHttpRequestPtr HttpRequest, FHttpResponsePtr HttpResponse, bool bSucceeded) {
 		CHECK(!bSucceeded);
 		CHECK(HttpResponse == nullptr);
 		CHECK(HttpRequest->GetStatus() == EHttpRequestStatus::Failed);
 		CHECK(HttpRequest->GetFailureReason() == EHttpFailureReason::ConnectionError);
 		const double DurationInSeconds  = FPlatformTime::Seconds() - StartTime;
-		double HttpTimeDiffTolerance = 1.5;
-#if WITH_CURL_XCURL
-		HttpTimeDiffTolerance += 3.0; // It seems xCurl takes up to 3 more seconds for connect timeout
-#endif
-		CHECK(FMath::IsNearlyEqual(DurationInSeconds, 15.0, HttpTimeDiffTolerance));
+		CHECK(FMath::IsNearlyEqual(DurationInSeconds, ExpectedTimeoutDuration, UE_HTTP_CONNECTION_TIMEOUT_MAX_DEVIATION));
 	});
 	HttpRequest->ProcessRequest();
 }
@@ -723,15 +749,15 @@ TEST_CASE_METHOD(FWaitUntilCompleteHttpFixture, "Http request activity timeout",
 		CHECK(HttpRequest->GetFailureReason() == EHttpFailureReason::ConnectionError);
 
 		const double DurationInSeconds  = FPlatformTime::Seconds() - StartTime;
-#if WITH_CURL_XCURL || PLATFORM_MAC || PLATFORM_IOS
+#if UE_HTTP_ACTIVITY_TIMER_START_AFTER_RECEIVED_DATA
 		// Unlike libCurl, currently there is an issue in xCurl that it triggers CURLINFO_HEADER_OUT even if can't
 		// connect. Had to disable that code, make sure not to treat that event as connected
 		// In a similar way on MacOS/iOS we don't get any notification until some data is received
 		// So it takes 5s to receive the first chunk to be considered as connected, then start response timer and
 		// take 3s to response timeout
-		CHECK(FMath::IsNearlyEqual(DurationInSeconds, ReceiveTimeoutSetting + 5, HTTP_TIME_DIFF_TOLERANCE));
+		CHECK(FMath::IsNearlyEqual(DurationInSeconds, ReceiveTimeoutSetting + 5, HTTP_TIME_DIFF_TOLERANCE_OF_REQUEST));
 #else
-		CHECK(FMath::IsNearlyEqual(DurationInSeconds, ReceiveTimeoutSetting, HTTP_TIME_DIFF_TOLERANCE));
+		CHECK(FMath::IsNearlyEqual(DurationInSeconds, ReceiveTimeoutSetting, HTTP_TIME_DIFF_TOLERANCE_OF_REQUEST));
 #endif
 
 	});
@@ -751,7 +777,7 @@ TEST_CASE_METHOD(FWaitUntilCompleteHttpFixture, "Http request won't trigger acti
 	const double StartTime = FPlatformTime::Seconds();
 	HttpRequest->OnProcessRequestComplete().BindLambda([StartTime, TimeToWaitBeforeCancel](FHttpRequestPtr HttpRequest, FHttpResponsePtr HttpResponse, bool bSucceeded) {
 		const double DurationInSeconds  = FPlatformTime::Seconds() - StartTime;
-		CHECK(FMath::IsNearlyEqual(DurationInSeconds, TimeToWaitBeforeCancel, HTTP_TIME_DIFF_TOLERANCE));
+		CHECK(FMath::IsNearlyEqual(DurationInSeconds, TimeToWaitBeforeCancel, HTTP_TIME_DIFF_TOLERANCE_OF_REQUEST));
 		CHECK(!bSucceeded);
 		CHECK(HttpRequest->GetStatus() == EHttpRequestStatus::Failed);
 		CHECK(HttpRequest->GetFailureReason() == EHttpFailureReason::Cancelled);
@@ -800,7 +826,7 @@ TEST_CASE_METHOD(FWaitUntilCompleteHttpFixture, "Http request total timeout with
 		CHECK(HttpRequest->GetStatus() == EHttpRequestStatus::Failed);
 		CHECK(HttpRequest->GetFailureReason() == EHttpFailureReason::TimedOut);
 		const double DurationInSeconds  = FPlatformTime::Seconds() - StartTime;
-		CHECK(FMath::IsNearlyEqual(DurationInSeconds, TotalTimeoutSetting, HTTP_TIME_DIFF_TOLERANCE));
+		CHECK(FMath::IsNearlyEqual(DurationInSeconds, TotalTimeoutSetting, HTTP_TIME_DIFF_TOLERANCE_OF_REQUEST));
 	});
 	HttpRequest->ProcessRequest();
 }
@@ -849,7 +875,7 @@ TEST_CASE_METHOD(FWaitUntilCompleteHttpFixture, "Http request total timeout with
 		CHECK(HttpRequest->GetStatus() == EHttpRequestStatus::Failed);
 		CHECK(HttpRequest->GetFailureReason() == EHttpFailureReason::TimedOut);
 		const double DurationInSeconds  = FPlatformTime::Seconds() - StartTime;
-		CHECK(FMath::IsNearlyEqual(DurationInSeconds, TimeoutSetting, HTTP_TIME_DIFF_TOLERANCE));
+		CHECK(FMath::IsNearlyEqual(DurationInSeconds, TimeoutSetting, HTTP_TIME_DIFF_TOLERANCE_OF_REQUEST));
 	});
 	HttpRequest->ProcessRequest();
 }
@@ -949,6 +975,33 @@ TEST_CASE_METHOD(FWaitUntilQuitFromTestFixture, "Http request can be reused", HT
 	HttpRequest->ProcessRequest();
 }
 
+#if UE_HTTP_CONNECTION_TIMEOUT_SUPPORT_RETRY
+TEST_CASE_METHOD(FWaitUntilQuitFromTestFixture, "Make sure connection time out can work well for 2nd same http request", HTTP_TAG)
+{
+	DisableWarningsInThisTest();
+
+	TSharedRef<IHttpRequest> HttpRequest = CreateRequest();
+
+	float ConnectionTimeoutDuration = 2.0f;
+	HttpModule->HttpConnectionTimeout = ConnectionTimeoutDuration;
+
+	HttpRequest->SetURL(UrlWithInvalidPortToTestConnectTimeout());
+
+	const double StartTime = FPlatformTime::Seconds();
+
+	HttpRequest->OnProcessRequestComplete().BindLambda([this, StartTime, ConnectionTimeoutDuration](FHttpRequestPtr HttpRequest, FHttpResponsePtr HttpResponse, bool bSucceeded) {
+		TSharedRef<IHttpRequest> HttpRequest2 = CreateRequest();
+		HttpRequest2->SetURL(UrlWithInvalidPortToTestConnectTimeout());
+		HttpRequest2->OnProcessRequestComplete().BindLambda([this, StartTime, ConnectionTimeoutDuration](FHttpRequestPtr HttpRequest, FHttpResponsePtr HttpResponse, bool bSucceeded) {
+			bQuitRequested = true;
+			const double DurationInSeconds  = FPlatformTime::Seconds() - StartTime;
+			CHECK(FMath::IsNearlyEqual(DurationInSeconds, ConnectionTimeoutDuration * 2, UE_HTTP_CONNECTION_TIMEOUT_MAX_DEVIATION * 2));
+		});
+		HttpRequest2->ProcessRequest();
+	});
+	HttpRequest->ProcessRequest();
+}
+#endif // UE_HTTP_CONNECTION_TIMEOUT_SUPPORT_RETRY
 
 // Response shared ptr should be able to be kept by user code and valid to access without http request
 class FValidateResponseDependencyFixture : public FWaitUntilCompleteHttpFixture
@@ -1429,7 +1482,7 @@ TEST_CASE_METHOD(FWaitUntilCompleteHttpFixture, "Retry respect Retry-After heade
 	HttpRequest->OnProcessRequestComplete().BindLambda([RetryAfter, StartTime](FHttpRequestPtr HttpRequest, FHttpResponsePtr HttpResponse, bool bSucceeded) {
 		CHECK(bSucceeded);
 		const double DurationInSeconds  = FPlatformTime::Seconds() - StartTime;
-		CHECK(FMath::IsNearlyEqual(DurationInSeconds, (float)(RetryAfter), HTTP_TIME_DIFF_TOLERANCE));
+		CHECK(FMath::IsNearlyEqual(DurationInSeconds, (float)(RetryAfter), HTTP_TIME_DIFF_TOLERANCE_OF_REQUEST));
 	});
 
 	HttpRequest->ProcessRequest();
@@ -1467,7 +1520,7 @@ TEST_CASE_METHOD(FWaitUntilCompleteHttpFixture, "Request can time out during loc
 		CHECK(HttpResponse->GetResponseCode() == EHttpResponseCodes::TooManyRequests);
 		CHECK(HttpResponse->GetContentLength() > 0);
 		const double DurationInSeconds  = FPlatformTime::Seconds() - StartTime;
-		CHECK(FMath::IsNearlyEqual(DurationInSeconds, 1.0, HTTP_TIME_DIFF_TOLERANCE));
+		CHECK(FMath::IsNearlyEqual(DurationInSeconds, 1.0, HTTP_TIME_DIFF_TOLERANCE_OF_REQUEST));
 	});
 
 	HttpRequest->ProcessRequest();
@@ -1495,14 +1548,15 @@ TEST_CASE_METHOD(FWaitUntilCompleteHttpFixture, "Request can time out during ret
 	// Will be forwarded back in response
 	HttpRequest->SetHeader(TEXT("Retry-After"), FString::Format(TEXT("{0}"), { RetryAfter }));
 
+	++ExpectingExtraCallbacks;
 	HttpRequest->OnRequestWillRetry().BindLambda([this](FHttpRequestPtr Request, FHttpResponsePtr /*Response*/, float LockoutPeriod) {
+		--ExpectingExtraCallbacks;
 		// Now retry with a latency during request
-		Request->SetURL(UrlStreamDownload(2/*Chunks*/, HTTP_TEST_TIMEOUT_CHUNK_SIZE, 2/*ChunkLatency*/));
+		Request->SetURL(UrlStreamDownload(3/*Chunks*/, HTTP_TEST_TIMEOUT_CHUNK_SIZE, 2/*ChunkLatency*/));
 	});
 
 	const double StartTime = FPlatformTime::Seconds();
 	HttpRequest->OnProcessRequestComplete().BindLambda([StartTime](FHttpRequestPtr HttpRequest, FHttpResponsePtr HttpResponse, bool bSucceeded) {
-		CHECK(HttpRequest->GetResponse() == nullptr); // Timeout will return response with nullptr in request itself
 		CHECK(HttpRequest->GetStatus() == EHttpRequestStatus::Failed);
 		CHECK(HttpRequest->GetFailureReason() == EHttpFailureReason::TimedOut);
 
@@ -1514,7 +1568,7 @@ TEST_CASE_METHOD(FWaitUntilCompleteHttpFixture, "Request can time out during ret
 		CHECK(HttpResponse->GetResponseCode() == EHttpResponseCodes::TooManyRequests);
 		CHECK(HttpResponse->GetContentLength() > 0);
 		const double DurationInSeconds  = FPlatformTime::Seconds() - StartTime;
-		CHECK(FMath::IsNearlyEqual(DurationInSeconds, 3.0, HTTP_TIME_DIFF_TOLERANCE));
+		CHECK(FMath::IsNearlyEqual(DurationInSeconds, 3.0, HTTP_TIME_DIFF_TOLERANCE_OF_REQUEST));
 	});
 
 	HttpRequest->ProcessRequest();
@@ -1552,6 +1606,7 @@ TEST_CASE_METHOD(FWaitUntilCompleteHttpFixture, "Retry immediately without lock
 	HttpRequest->ProcessRequest();
 }
 
+#if UE_HTTP_CONNECTION_TIMEOUT_SUPPORT_RETRY
 TEST_CASE_METHOD(FWaitUntilCompleteHttpFixture, "Optionally retry limit can be set differently for connection error", HTTP_TAG)
 {
 	if (!bRetryEnabled)
@@ -1561,8 +1616,10 @@ TEST_CASE_METHOD(FWaitUntilCompleteHttpFixture, "Optionally retry limit can be s
 
 	DisableWarningsInThisTest();
 
+	HttpModule->HttpConnectionTimeout = 1.0f;
+
 	TSharedRef<IHttpRequest> HttpRequest = HttpRetryManager->CreateRequest(
-		5/*InRetryLimitCountOverride*/,
+		3/*InRetryLimitCountOverride*/,
 		FHttpRetrySystem::FRetryTimeoutRelativeSecondsSetting()/*InRetryTimeoutRelativeSecondsOverride unused*/,
 		{ EHttpResponseCodes::TooManyRequests, EHttpResponseCodes::ServiceUnavail }/*InRetryResponseCodes*/,
 		FHttpRetrySystem::FRetryVerbs(), /*unused*/
@@ -1571,37 +1628,31 @@ TEST_CASE_METHOD(FWaitUntilCompleteHttpFixture, "Optionally retry limit can be s
 	);
 
 	float ExpectedTimeoutDuration = 0.0f;
-	SECTION("RetryLimitCountDefault:5 will be used so retries in general take long")
-	{
-		HttpRequest->SetURL(UrlMockStatus(EHttpResponseCodes::TooManyRequests));
-		HttpRequest->SetHeader(TEXT("Retry-After"), FString::Format(TEXT("{0}"), { 3 }));
-
-		ExpectedTimeoutDuration = 15.0f; // each request will take about 0s, 5 retry back offs, each back off takes 3s;
-	}
+	float TimeDiffTolerance = 0.0f;
 	SECTION("RetryLimitCountForConnectionErrorDefault:1 will be used so retries for connection error take less time")
 	{
-		HttpRequest->SetURL(UrlStreamDownload(2/*Chunks*/, HTTP_TEST_TIMEOUT_CHUNK_SIZE, 2/*ChunkLatency*/));
-		HttpModule->HttpActivityTimeout = 1.0f;
+		HttpRequest->SetURL(UrlWithInvalidPortToTestConnectTimeout());
 
 		ExpectedTimeoutDuration = 2.0f; // each request will take 1s, 1st retry back off takes 0s
+		TimeDiffTolerance = 2 * UE_HTTP_CONNECTION_TIMEOUT_MAX_DEVIATION;
+	}
+	SECTION("RetryLimitCountDefault:3 will be used so retries in general take long")
+	{
+		HttpRequest->SetURL(UrlMockStatus(EHttpResponseCodes::TooManyRequests));
+		HttpRequest->SetHeader(TEXT("Retry-After"), FString::Format(TEXT("{0}"), { 3 }));
 
-#if WITH_CURL_XCURL
-		// Unlike libCurl, currently there is an issue in xCurl that it triggers CURLINFO_HEADER_OUT even if can't 
-		// connect. Had to disable that code, make sure not to treat that event as connected
-		// It takes 2s to receive the first chunk to be considered as connected, then start response timer and 
-		// take 1s to time out. So it takes 3s in total for each request instead of 1s
-		ExpectedTimeoutDuration += 4;
-#endif
+		ExpectedTimeoutDuration = 9.0f; // each request will take 0s, 3 retry back offs, each back off takes 3s;
+		TimeDiffTolerance = 3 * HTTP_TIME_DIFF_TOLERANCE_OF_REQUEST;
 	}
 
 	const double StartTime = FPlatformTime::Seconds();
-	HttpRequest->OnProcessRequestComplete().BindLambda([StartTime, ExpectedTimeoutDuration](FHttpRequestPtr HttpRequest, FHttpResponsePtr HttpResponse, bool bSucceeded) {
+	HttpRequest->OnProcessRequestComplete().BindLambda([StartTime, ExpectedTimeoutDuration, TimeDiffTolerance](FHttpRequestPtr HttpRequest, FHttpResponsePtr HttpResponse, bool bSucceeded) {
 		const double DurationInSeconds  = FPlatformTime::Seconds() - StartTime;
-		float TimeDiffTolerance = 1.0f;
 		CHECK(FMath::IsNearlyEqual(DurationInSeconds, ExpectedTimeoutDuration, TimeDiffTolerance));
 	});
 	HttpRequest->ProcessRequest();
 }
+#endif // UE_HTTP_CONNECTION_TIMEOUT_SUPPORT_RETRY
 
 TEST_CASE_METHOD(FWaitUntilCompleteHttpFixture, "Retry fallback with exponential lock out if there is no Retry-After header", HTTP_TAG)
 {
@@ -1697,7 +1748,7 @@ TEST_CASE_METHOD(FWaitUntilCompleteHttpFixture, "Flush while activity timeout sh
 	HttpModule->GetHttpManager().Flush(EHttpFlushReason::FullFlush);
 }
 
-#if (PLATFORM_WINDOWS && !WITH_CURL_XCURL) || PLATFORM_MAC || PLATFORM_UNIX
+#if UE_HTTP_SUPPORT_LOCAL_SERVER
 
 TEST_CASE_METHOD(FWaitUntilCompleteHttpFixture, "Scheme besides http and https can work if allowed by settings", HTTP_TAG)
 {
@@ -1793,7 +1844,7 @@ TEST_CASE_METHOD(FLocalHttpServerFixture, "Local http server can serve large fil
 	HttpRequest->ProcessRequest();
 }
 
-#endif // (PLATFORM_WINDOWS && !WITH_CURL_XCURL) || PLATFORM_MAC || PLATFORM_UNIX
+#endif // UE_HTTP_SUPPORT_LOCAL_SERVER
 
 TEST_CASE_METHOD(FWaitUntilCompleteHttpFixture, "Test platform request requests limits", HTTP_TAG "[LIMIT]")
 {
diff --git a/Engine/Source/Runtime/Online/HTTP/HTTP.Build.cs b/Engine/Source/Runtime/Online/HTTP/HTTP.Build.cs
--- a/Engine/Source/Runtime/Online/HTTP/HTTP.Build.cs
+++ b/Engine/Source/Runtime/Online/HTTP/HTTP.Build.cs
@@ -28,6 +28,17 @@ public class HTTP : ModuleRules
 
 	protected virtual bool bPlatformSupportsCurlMultiWait { get { return false; } }
 	protected virtual bool bPlatformSupportsCurlQuickExit { get { return !bPlatformSupportsXCurl; } }
+	protected virtual bool bPlatformConnectionTimeoutSupportRetry { get { return true; } }
+	protected virtual bool bPlatformSupportsLocalHttpServer 
+	{ 
+		get 
+		{ 
+			return !bPlatformSupportsXCurl && 
+				!Target.IsInPlatformGroup(UnrealPlatformGroup.Android) && 
+				!Target.IsInPlatformGroup(UnrealPlatformGroup.IOS); 
+		} 
+	}
+
 	protected virtual int DefaultMaxConcurrentRequests { get { return int.MaxValue; } }
 
 	private bool bPlatformSupportsCurl { get { return bPlatformSupportsLibCurl || bPlatformSupportsXCurl; } }
@@ -83,7 +94,7 @@ public class HTTP : ModuleRules
 		}
 
 		PrivateDefinitions.Add("WITH_CURL_LIBCURL =" + (bPlatformSupportsLibCurl ? "1" : "0"));
-		PublicDefinitions.Add("WITH_CURL_XCURL=" + (bPlatformSupportsXCurl ? "1" : "0"));
+		PrivateDefinitions.Add("WITH_CURL_XCURL=" + (bPlatformSupportsXCurl ? "1" : "0"));
 		PrivateDefinitions.Add("WITH_CURL_MULTIPOLL=" + (bPlatformSupportsCurlMultiPoll ? "1" : "0"));
 		PrivateDefinitions.Add("WITH_CURL_MULTIWAIT=" + (bPlatformSupportsCurlMultiWait ? "1" : "0"));
 		PrivateDefinitions.Add("WITH_CURL_MULTISOCKET=" + (bPlatformSupportsCurlMultiSocket ? "1" : "0"));
@@ -117,5 +128,19 @@ public class HTTP : ModuleRules
 		}
 
 		PrivateDefinitions.Add("UE_HTTP_DEFAULT_MAX_CONCURRENT_REQUESTS=" + DefaultMaxConcurrentRequests);
+
+		float PlatformConnectionTimeoutMaxDeviation = 0.5f;
+		if (bPlatformSupportsXCurl)
+		{
+			PlatformConnectionTimeoutMaxDeviation = 3.5f;
+		}
+		if (Target.IsInPlatformGroup(UnrealPlatformGroup.Apple))
+		{
+			PlatformConnectionTimeoutMaxDeviation = 1.5f;
+		}
+		PublicDefinitions.Add("UE_HTTP_CONNECTION_TIMEOUT_MAX_DEVIATION=" + PlatformConnectionTimeoutMaxDeviation);
+		PublicDefinitions.Add("UE_HTTP_CONNECTION_TIMEOUT_SUPPORT_RETRY=" + (bPlatformConnectionTimeoutSupportRetry ? "1" : "0"));
+		PublicDefinitions.Add("UE_HTTP_ACTIVITY_TIMER_START_AFTER_RECEIVED_DATA=" + ((bPlatformSupportsXCurl || Target.IsInPlatformGroup(UnrealPlatformGroup.Apple)) ? "1" : "0"));
+		PublicDefinitions.Add("UE_HTTP_SUPPORT_LOCAL_SERVER=" + (bPlatformSupportsLocalHttpServer ? "1" : "0"));
 	}
 }
