Fix the test case "Optionally retry limit can be set differently for connection error" by using latest expected back off value;
Make the test case "Test platform request requests limits" covers the cancelling using CompleteOnHttpThread as delegate thread policy.

#tests Passed the WebTests.

diff --git a/Engine/Source/Programs/WebTests/Private/TestHttp.cpp b/Engine/Source/Programs/WebTests/Private/TestHttp.cpp
--- a/Engine/Source/Programs/WebTests/Private/TestHttp.cpp
+++ b/Engine/Source/Programs/WebTests/Private/TestHttp.cpp
@@ -1713,13 +1713,17 @@ TEST_CASE_METHOD(FWaitUntilCompleteHttpFixture, "Optionally retry limit can be s
 
 	HttpModule->HttpConnectionTimeout = 1.0f;
 
+	FHttpRetrySystem::FExponentialBackoffCurve RetryBackoffCurve;
+	RetryBackoffCurve.MinCoefficient = 1.0f; // no jitter
+
 	TSharedRef<IHttpRequest> HttpRequest = HttpRetryManager->CreateRequest(
 		3/*InRetryLimitCountOverride*/,
 		FHttpRetrySystem::FRetryTimeoutRelativeSecondsSetting()/*InRetryTimeoutRelativeSecondsOverride unused*/,
 		{ EHttpResponseCodes::TooManyRequests, EHttpResponseCodes::ServiceUnavail }/*InRetryResponseCodes*/,
 		FHttpRetrySystem::FRetryVerbs(), /*unused*/
 		FHttpRetrySystem::FRetryDomainsPtr(), /*unused*/
-		1 /*InRetryLimitCountForConnectionErrorOverride*/
+		1, /*InRetryLimitCountForConnectionErrorOverride*/
+		RetryBackoffCurve/*InExponentialBackoffCurve*/
 	);
 
 	float ExpectedTimeoutDuration = 0.0f;
@@ -1728,7 +1732,7 @@ TEST_CASE_METHOD(FWaitUntilCompleteHttpFixture, "Optionally retry limit can be s
 	{
 		HttpRequest->SetURL(UrlWithInvalidPortToTestConnectTimeout());
 
-		ExpectedTimeoutDuration = 2.0f; // each request will take 1s, 1st retry back off takes 0s
+		ExpectedTimeoutDuration = 6.0f; // each request will take 1s, 1st retry back off takes 4s
 		TimeDiffTolerance = 2 * UE_HTTP_CONNECTION_TIMEOUT_MAX_DEVIATION;
 	}
 	SECTION("RetryLimitCountDefault:3 will be used so retries in general take long")
@@ -1985,8 +1989,15 @@ TEST_CASE_METHOD(FWaitUntilCompleteHttpFixture, "Test platform request requests
 		{
 			TSharedRef<IHttpRequest> HttpRequest = FHttpModule::Get().CreateRequest();
 			// Requests server to serve 1024b chunks to allow time for cancel to happen
-			HttpRequest->SetURL(UrlStreamDownload(2, HTTP_TEST_TIMEOUT_CHUNK_SIZE, /*ChunkLatency=*/bCheckCancel ? 1 : 0));
+			HttpRequest->SetURL(UrlStreamDownload(3, HTTP_TEST_TIMEOUT_CHUNK_SIZE, /*ChunkLatency=*/bCheckCancel ? 1 : 0));
 			HttpRequest->SetVerb(TEXT("GET"));
+
+			// Since catch2 uses std::srand, use std::rand here should make it deterministic when use same seed through --rng-seed
+			if (std::rand() % 2)
+			{
+				HttpRequest->SetDelegateThreadPolicy(EHttpRequestDelegateThreadPolicy::CompleteOnHttpThread);
+			}
+
 			HttpRequest->OnProcessRequestComplete().BindLambda([bCheckCancel](FHttpRequestPtr HttpRequest, FHttpResponsePtr HttpResponse, bool bSucceeded)
 			{
 				//Only assert if response is successful on non-canceled requests
@@ -2005,6 +2016,9 @@ TEST_CASE_METHOD(FWaitUntilCompleteHttpFixture, "Test platform request requests
 
 		if(bCheckCancel)
 		{
+			// Make sure requests are started in http thread
+			FPlatformProcess::Sleep(0.1);
+
 			for (auto Request : Requests)
 			{
 				Request->CancelRequest();
diff --git a/Engine/Source/Runtime/Online/HTTP/Private/HttpRetrySystem.cpp b/Engine/Source/Runtime/Online/HTTP/Private/HttpRetrySystem.cpp
--- a/Engine/Source/Runtime/Online/HTTP/Private/HttpRetrySystem.cpp
+++ b/Engine/Source/Runtime/Online/HTTP/Private/HttpRetrySystem.cpp
@@ -343,7 +343,8 @@ TSharedRef<FHttpRetrySystem::FRequest, ESPMode::ThreadSafe> FHttpRetrySystem::FM
 	const FRetryResponseCodes& InRetryResponseCodes,
 	const FRetryVerbs& InRetryVerbs,
 	const FRetryDomainsPtr& InRetryDomains,
-	const FRetryLimitCountSetting& InRetryLimitCountForConnectionErrorOverride
+	const FRetryLimitCountSetting& InRetryLimitCountForConnectionErrorOverride,
+	const FExponentialBackoffCurve& InExponentialBackoffCurve
 )
 {
 	return MakeShareable(new FRequest(
@@ -354,7 +355,8 @@ TSharedRef<FHttpRetrySystem::FRequest, ESPMode::ThreadSafe> FHttpRetrySystem::FM
 		InRetryResponseCodes,
 		InRetryVerbs,
 		InRetryDomains,
-		InRetryLimitCountForConnectionErrorOverride
+		InRetryLimitCountForConnectionErrorOverride,
+		InExponentialBackoffCurve
 		));
 }
 
diff --git a/Engine/Source/Runtime/Online/HTTP/Public/HttpRetrySystem.h b/Engine/Source/Runtime/Online/HTTP/Public/HttpRetrySystem.h
--- a/Engine/Source/Runtime/Online/HTTP/Public/HttpRetrySystem.h
+++ b/Engine/Source/Runtime/Online/HTTP/Public/HttpRetrySystem.h
@@ -194,7 +194,8 @@ public:
 		const FRetryResponseCodes& InRetryResponseCodes = FRetryResponseCodes(),
 		const FRetryVerbs& InRetryVerbs = FRetryVerbs(),
 		const FRetryDomainsPtr& InRetryDomains = FRetryDomainsPtr(),
-		const FRetryLimitCountSetting& InRetryLimitCountForConnectionErrorOverride = FRetryLimitCountSetting()
+		const FRetryLimitCountSetting& InRetryLimitCountForConnectionErrorOverride = FRetryLimitCountSetting(),
+		const FExponentialBackoffCurve & InExponentialBackoffCurve = FExponentialBackoffCurve()
 	);
 
 	HTTP_API virtual ~FManager();
