Make sure to stop activity timeout timer when total timeout triggered;
Only output log if request had completed when activity timeout, it can continue to run even if that happened.

diff --git a/Engine/Source/Programs/WebTests/Private/TestHttp.cpp b/Engine/Source/Programs/WebTests/Private/TestHttp.cpp
--- a/Engine/Source/Programs/WebTests/Private/TestHttp.cpp
+++ b/Engine/Source/Programs/WebTests/Private/TestHttp.cpp
@@ -62,10 +62,7 @@ public:
 
 	~FHttpTestLogLevelInitializer()
 	{
-		if (OldVerbosity != LogHttp.GetVerbosity())
-		{
-			LogHttp.SetVerbosity(OldVerbosity);
-		}
+		ResumeLogVerbosity();
 	}
 
 	void DisableWarningsInThisTest()
@@ -76,6 +73,14 @@ public:
 		}
 	}
 
+	void ResumeLogVerbosity()
+	{
+		if (OldVerbosity != LogHttp.GetVerbosity())
+		{
+			LogHttp.SetVerbosity(OldVerbosity);
+		}
+	}
+
 	bool bVeryVerbose = false;
 	ELogVerbosity::Type OldVerbosity;
 };
@@ -119,6 +124,11 @@ public:
 		HttpTestLogLevelInitializer.DisableWarningsInThisTest();
 	}
 
+	void ResumeLogVerbosity()
+	{
+		HttpTestLogLevelInitializer.ResumeLogVerbosity();
+	}
+
 	const FString UrlWithInvalidPortToTestConnectTimeout() const { return TEXT("http://10.255.255.1:8765"); } // non-routable IP address with a random port
 	const FString UrlBase() const { return FString::Format(TEXT("http://{0}:{1}"), { *WebServerIp, WebServerHttpPort }); }
 	const FString UrlHttpTests() const { return FString::Format(TEXT("{0}/webtests/httptests"), { *UrlBase() }); }
@@ -883,6 +893,29 @@ TEST_CASE_METHOD(FWaitUntilCompleteHttpFixture, "Http request won't trigger acti
 	FPlatformProcess::Sleep(3.0f); // Just make sure there is no warning or assert triggered by the activity timeout callback
 }
 
+TEST_CASE_METHOD(FWaitUntilCompleteHttpFixture, "Http request won't trigger activity timeout after total timeout", HTTP_TAG)
+{
+	DisableWarningsInThisTest();
+
+	HttpModule->HttpActivityTimeout = 2.0f;
+	HttpModule->HttpTotalTimeout = 3.5f;
+
+	TSharedPtr<IHttpRequest> HttpRequest = CreateRequest();
+	HttpRequest->SetURL(UrlStreamDownload(5/*Chunks*/, HTTP_TEST_TIMEOUT_CHUNK_SIZE, 1/*ChunkLatency*/));
+	HttpRequest->SetVerb(TEXT("GET"));
+	HttpRequest->SetDelegateThreadPolicy(EHttpRequestDelegateThreadPolicy::CompleteOnHttpThread);
+
+	const double StartTime = FPlatformTime::Seconds();
+	HttpRequest->OnProcessRequestComplete().BindLambda([this](FHttpRequestPtr HttpRequest, FHttpResponsePtr HttpResponse, bool bSucceeded) {
+		CHECK(!bSucceeded);
+		CHECK(HttpRequest->GetStatus() == EHttpRequestStatus::Failed);
+		CHECK(HttpRequest->GetFailureReason() == EHttpFailureReason::TimedOut);
+		ResumeLogVerbosity();
+	});
+	HttpRequest->ProcessRequest();
+	FPlatformProcess::Sleep(6.0f); // Just make sure there is no warning or assert triggered by the activity timeout callback
+}
+
 TEST_CASE_METHOD(FWaitUntilCompleteHttpFixture, "Http request receive won't timeout for streaming request", HTTP_TAG)
 {
 	HttpModule->HttpActivityTimeout = 3.0f;
diff --git a/Engine/Source/Runtime/Online/HTTP/Private/GenericPlatform/HttpRequestCommon.cpp b/Engine/Source/Runtime/Online/HTTP/Private/GenericPlatform/HttpRequestCommon.cpp
--- a/Engine/Source/Runtime/Online/HTTP/Private/GenericPlatform/HttpRequestCommon.cpp
+++ b/Engine/Source/Runtime/Online/HTTP/Private/GenericPlatform/HttpRequestCommon.cpp
@@ -285,7 +285,11 @@ void FHttpRequestCommon::OnActivityTimeoutTimerTaskTrigger()
 
 	ActivityTimeoutHttpTaskTimerHandle.Reset();
 
-	checkf(!EHttpRequestStatus::IsFinished(GetStatus()), TEXT("PostProcess must be called when request complete, to stop activity timeout timer."));
+	if (EHttpRequestStatus::IsFinished(GetStatus()))
+	{
+		UE_LOG(LogHttp, Warning, TEXT("Request %p had finished when activity timeout timer trigger at [%s]"), this, *FDateTime::Now().ToString(TEXT("%H:%M:%S:%s")));
+		return;
+	}
 
 	if (FPlatformTime::Seconds() < ActivityTimeoutAt)
 	{
@@ -381,12 +385,17 @@ void FHttpRequestCommon::OnTotalTimeoutTimerTaskTrigger()
 	const FScopeLock CacheLock(&HttpTaskTimerHandleCriticalSection);
 	bTimedOut = true;
 
-	if (!EHttpRequestStatus::IsFinished(GetStatus())) 
+	if (EHttpRequestStatus::IsFinished(GetStatus()))
 	{
-		QUICK_SCOPE_CYCLE_COUNTER(STAT_FHttpRequestCommon_AbortRequest);
-		UE_LOG(LogHttp, Warning, TEXT("HTTP request timed out after %0.2f seconds URL=%s"), GetTimeoutOrDefault(), *GetURL());
-		AbortRequest();
+		return;
 	}
+
+	StopActivityTimeoutTimer();
+
+	QUICK_SCOPE_CYCLE_COUNTER(STAT_FHttpRequestCommon_AbortRequest);
+	UE_LOG(LogHttp, Warning, TEXT("HTTP request timed out after %0.2f seconds URL=%s"), GetTimeoutOrDefault(), *GetURL());
+
+	AbortRequest();
 }
 
 void FHttpRequestCommon::StopTotalTimeoutTimer()
