When http request total timeout during lockout, make sure it can trigger complete delegate by http delegate thread policy;
Use FTSTicker in http manager instead to manage game thread tasks.

#tests Passed all tests in WebTests project

diff --git a/Engine/Source/Programs/WebTests/Private/TestHttp.cpp b/Engine/Source/Programs/WebTests/Private/TestHttp.cpp
--- a/Engine/Source/Programs/WebTests/Private/TestHttp.cpp
+++ b/Engine/Source/Programs/WebTests/Private/TestHttp.cpp
@@ -224,17 +224,26 @@ public:
 		ensure(--OngoingRequests >= 0);
 	}
 
+	void TickHttpManager()
+	{
+		double Now = FPlatformTime::Seconds();
+		static double LastTick = Now;
+		double Duration = Now - LastTick;
+		LastTick = Now;
+		HttpModule->GetHttpManager().Tick(Duration);
+		FPlatformProcess::Sleep(TickFrequency);
+	}
+
 	void WaitUntilAllHttpRequestsComplete()
 	{
 		while (HasOngoingRequest())
 		{
-			HttpModule->GetHttpManager().Tick(TickFrequency);
-			FPlatformProcess::Sleep(TickFrequency);
+			TickHttpManager();
 		}
 
 		// In case in http thread the http request complete and set OngoingRequests to 0, http manager never 
 		// had chance to Tick and remove the request
-		HttpModule->GetHttpManager().Tick(TickFrequency);
+		TickHttpManager();
 	}
 
 	bool HasOngoingRequest() const
@@ -1059,8 +1068,7 @@ public:
 	{
 		while (!bQuitRequested)
 		{
-			HttpModule->GetHttpManager().Tick(TickFrequency);
-			FPlatformProcess::Sleep(TickFrequency);
+			TickHttpManager();
 		}
 	}
 
@@ -1625,6 +1633,15 @@ TEST_CASE_METHOD(FWaitUntilCompleteHttpFixture, "Request can time out during loc
 
 	DisableWarningsInThisTest();
 
+	EHttpRequestDelegateThreadPolicy ThreadPolicyExpected = EHttpRequestDelegateThreadPolicy::CompleteOnGameThread;
+	SECTION("From game thread")
+	{
+	}
+	SECTION("From http thread")
+	{
+		ThreadPolicyExpected = EHttpRequestDelegateThreadPolicy::CompleteOnHttpThread;
+	}
+
 	TSharedRef<IHttpRequest> HttpRequest = HttpRetryManager->CreateRequest(
 		1/*InRetryLimitCountOverride*/,
 		FHttpRetrySystem::FRetryTimeoutRelativeSecondsSetting()/*InRetryTimeoutRelativeSecondsOverride unused*/,
@@ -1633,6 +1650,7 @@ TEST_CASE_METHOD(FWaitUntilCompleteHttpFixture, "Request can time out during loc
 
 	HttpRequest->SetURL(UrlMockStatus(EHttpResponseCodes::TooManyRequests));
 	HttpRequest->SetTimeout(1.0f);
+	HttpRequest->SetDelegateThreadPolicy(ThreadPolicyExpected);
 
 	uint32 RetryAfter = 4;
 
@@ -1640,7 +1658,7 @@ TEST_CASE_METHOD(FWaitUntilCompleteHttpFixture, "Request can time out during loc
 	HttpRequest->SetHeader(TEXT("Retry-After"), FString::Format(TEXT("{0}"), { RetryAfter }));
 
 	const double StartTime = FPlatformTime::Seconds();
-	HttpRequest->OnProcessRequestComplete().BindLambda([StartTime](FHttpRequestPtr HttpRequest, FHttpResponsePtr HttpResponse, bool bSucceeded) {
+	HttpRequest->OnProcessRequestComplete().BindLambda([StartTime, ThreadPolicyExpected](FHttpRequestPtr HttpRequest, FHttpResponsePtr HttpResponse, bool bSucceeded) {
 		// When timeout during lock out period, it fails with result of last request before lock out
 		CHECK(bSucceeded);
 		REQUIRE(HttpResponse != nullptr);
@@ -1649,6 +1667,7 @@ TEST_CASE_METHOD(FWaitUntilCompleteHttpFixture, "Request can time out during loc
 		CHECK(HttpResponse->GetContentLength() > 0);
 		const double DurationInSeconds  = FPlatformTime::Seconds() - StartTime;
 		CHECK(FMath::IsNearlyEqual(DurationInSeconds, 1.0, HTTP_TIME_DIFF_TOLERANCE_OF_REQUEST));
+		CHECK((ThreadPolicyExpected == EHttpRequestDelegateThreadPolicy::CompleteOnGameThread && IsInGameThread() || ThreadPolicyExpected == EHttpRequestDelegateThreadPolicy::CompleteOnHttpThread && !IsInGameThread()));
 	});
 
 	HttpRequest->ProcessRequest();
diff --git a/Engine/Source/Runtime/Online/HTTP/Private/HttpManager.cpp b/Engine/Source/Runtime/Online/HTTP/Private/HttpManager.cpp
--- a/Engine/Source/Runtime/Online/HTTP/Private/HttpManager.cpp
+++ b/Engine/Source/Runtime/Online/HTTP/Private/HttpManager.cpp
@@ -272,11 +272,14 @@ void FHttpManager::UpdateConfigs()
 	}
 }
 
-void FHttpManager::AddGameThreadTask(TFunction<void()>&& Task)
+void FHttpManager::AddGameThreadTask(TFunction<void()>&& Task, float Delay)
 {
 	if (Task)
 	{
-		GameThreadQueue.Enqueue(MoveTemp(Task));
+		GameThreadTicker.AddTicker(FTickerDelegate::CreateLambda([Task](float DeltaTime) {
+			Task();
+			return false;
+		}), Delay);
 	}
 }
 
@@ -425,14 +428,8 @@ bool FHttpManager::Tick(float DeltaSeconds)
 
 	// Run GameThread tasks
 	{
-		FScopeLock ScopeLock(&GameThreadQueueLock);
-
-		TFunction<void()> Task = nullptr;
-		while (GameThreadQueue.Dequeue(Task))
-		{
-			check(Task);
-			Task();
-		}
+		FScopeLock ScopeLock(&GameThreadTickerLock);
+		GameThreadTicker.Tick(DeltaSeconds);
 	}
 
 	if (Thread)
diff --git a/Engine/Source/Runtime/Online/HTTP/Private/HttpRetrySystem.cpp b/Engine/Source/Runtime/Online/HTTP/Private/HttpRetrySystem.cpp
--- a/Engine/Source/Runtime/Online/HTTP/Private/HttpRetrySystem.cpp
+++ b/Engine/Source/Runtime/Online/HTTP/Private/HttpRetrySystem.cpp
@@ -482,7 +482,7 @@ void FHttpRetrySystem::FManager::RetryHttpRequestWithDelay(FManager::FHttpRetryR
 		float WillTimeoutInDelay = TimeoutOrDefault - TimeElapsedForTheRequest;
 		if (WillTimeoutInDelay < InDelay)
 		{
-			HttpRequestTimeoutAfterDelay(RequestEntry, bWasSucceeded, TimeoutOrDefault);
+			HttpRequestTimeoutAfterDelay(RequestEntry, bWasSucceeded, WillTimeoutInDelay);
 			return;
 		}
 	}
@@ -513,7 +513,7 @@ void FHttpRetrySystem::FManager::RetryHttpRequestWithDelay(FManager::FHttpRetryR
 void FHttpRetrySystem::FManager::HttpRequestTimeoutAfterDelay(FManager::FHttpRetryRequestEntry& RequestEntry, bool bWasSucceeded, float Delay)
 {
 	TWeakPtr<FRequest> RequestWeakPtr(RequestEntry.Request);
-	FHttpModule::Get().GetHttpManager().AddHttpThreadTask([RequestWeakPtr, bWasSucceeded]() {
+	TFunction<void()> Callback([RequestWeakPtr, bWasSucceeded]() {
 		if (TSharedPtr<FRequest> RequestPtr = RequestWeakPtr.Pin())
 		{
 			if (TSharedPtr<FManager> RetryManagerPtr = RequestPtr->RetryManager.Pin())
@@ -534,7 +534,16 @@ void FHttpRetrySystem::FManager::HttpRequestTimeoutAfterDelay(FManager::FHttpRet
 			// Same as existing behavior, when timeout during lock out period, it fails with result of last request before lockout
 			RequestPtr->OnProcessRequestComplete().ExecuteIfBound(RequestPtr, RequestPtr->GetResponse(), bWasSucceeded);
 		}
-	}, Delay);
+	});
+
+	if (RequestEntry.Request->GetDelegateThreadPolicy() == EHttpRequestDelegateThreadPolicy::CompleteOnGameThread)
+	{
+		FHttpModule::Get().GetHttpManager().AddGameThreadTask(MoveTemp(Callback), Delay);
+	}
+	else
+	{
+		FHttpModule::Get().GetHttpManager().AddHttpThreadTask(MoveTemp(Callback), Delay);
+	}
 }
 
 float FHttpRetrySystem::FManager::GetLockoutPeriodSeconds(const FHttpRetryRequestEntry& HttpRetryRequestEntry)
diff --git a/Engine/Source/Runtime/Online/HTTP/Public/HttpManager.h b/Engine/Source/Runtime/Online/HTTP/Public/HttpManager.h
--- a/Engine/Source/Runtime/Online/HTTP/Public/HttpManager.h
+++ b/Engine/Source/Runtime/Online/HTTP/Public/HttpManager.h
@@ -246,7 +246,7 @@ public:
 	 *
 	 * @param Task The task to be ran next tick
 	 */
-	HTTP_API void AddGameThreadTask(TFunction<void()>&& Task);
+	HTTP_API void AddGameThreadTask(TFunction<void()>&& Task, float Delay = 0.0f);
 
 	/**
 	 * Add task to be ran on the http thread
@@ -294,9 +294,9 @@ protected:
 	/** This method will be called to generate a CorrelationId on all requests being sent if one is not already set */
 	TFunction<FString()> CorrelationIdMethod;
 
-	/** Queue of tasks to run on the game thread */
-	TQueue<TFunction<void()>, EQueueMode::Mpsc> GameThreadQueue;
-	FCriticalSection GameThreadQueueLock;
+	/** Ticker to run game thread tasks*/
+	FTSTicker GameThreadTicker;
+	FCriticalSection GameThreadTickerLock;
 
 	// This variable is set to true in Flush(EHttpFlushReason), and prevents new Http requests from being launched
 	bool bFlushing;
