Fix a random dead lock when shut down http manager.

Reason: Currently when shutdown and flush in http manager, it gives up and cancels remaining requests after waiting a short period. When cancel, it holds the RequestLock, and in FHttpRequestCommon::CancelRequest, it tries to access HttpTaskTimerHandleCriticalSection to stop activity timer. In the mean time, the request itself could be in activity or total timeout callback, which is holding the HttpTaskTimerHandleCriticalSection, and trying to AbortRequest which will try to access RequstLock in HttpManager.IsValidRequest.

#tests Reproduced with test case in WebTests and fixed it. Tried launch and quit game on Win64.

diff --git a/Engine/Source/Programs/WebTests/Private/TestHttp.cpp b/Engine/Source/Programs/WebTests/Private/TestHttp.cpp
--- a/Engine/Source/Programs/WebTests/Private/TestHttp.cpp
+++ b/Engine/Source/Programs/WebTests/Private/TestHttp.cpp
@@ -16,7 +16,6 @@
 #include "Misc/CommandLine.h"
 #include "TestHarness.h"
 #include "Serialization/JsonSerializerMacros.h"
-
 #include "catch2/generators/catch_generators.hpp"
 
 /**
@@ -35,6 +34,7 @@
 
 #define HTTP_TAG "[HTTP]"
 #define HTTP_TIME_DIFF_TOLERANCE_OF_REQUEST 0.5f
+#define HTTP_TEST_TIMEOUT_CHUNK_SIZE 16*1024 // Use a big chunk size so it triggers data received callback in time on all platforms
 
 extern TAutoConsoleVariable<bool> CVarHttpInsecureProtocolEnabled;
 extern TAutoConsoleVariable<bool> CVarHttpRetrySystemNonGameThreadSupportEnabled;
@@ -85,6 +85,23 @@ public:
 	ELogVerbosity::Type OldVerbosity;
 };
 
+class FMockRetryManager : public FHttpRetrySystem::FManager
+{
+public:
+	using FHttpRetrySystem::FManager::FManager;
+	using FHttpRetrySystem::FManager::RequestList;
+	using FHttpRetrySystem::FManager::FHttpRetryRequestEntry;
+	using FHttpRetrySystem::FManager::RetryTimeoutRelativeSecondsDefault;
+	using FHttpRetrySystem::FManager::RetryLimitCountDefault;
+	using FHttpRetrySystem::FManager::RetryLimitCountForConnectionErrorDefault;
+
+	bool IsEmpty()
+	{
+		FScopeLock ScopeLock(&RequestListLock);
+		return RequestList.IsEmpty();
+	}
+};
+
 class FHttpModuleTestFixture
 {
 public:
@@ -98,18 +115,38 @@ public:
 
 		bRetryEnabled &= CVarHttpRetrySystemNonGameThreadSupportEnabled.GetValueOnAnyThread();
 
+		InitModule();
+
+		CVarHttpInsecureProtocolEnabled->Set(true);
+	}
+
+	void InitModule()
+	{
 		HttpModule = new FMockHttpModule();
 		IModuleInterface* Module = HttpModule;
 		Module->StartupModule();
+		if (bRetryEnabled)
+		{
+			HttpRetryManager = MakeShared<FMockRetryManager>(FHttpRetrySystem::FRetryLimitCountSetting(0), FHttpRetrySystem::FRetryTimeoutRelativeSecondsSetting(/*RetryTimeoutRelativeSeconds*/));
+		}
+	}
 
-		CVarHttpInsecureProtocolEnabled->Set(true);
+	void ShutdownModule()
+	{
+		HttpRetryManager = nullptr;
+
+		IModuleInterface* Module = HttpModule;
+		if (Module)
+		{
+			Module->ShutdownModule();
+			delete Module;
+			HttpModule = nullptr;
+		}
 	}
 
 	virtual ~FHttpModuleTestFixture()
 	{
-		IModuleInterface* Module = HttpModule;
-		Module->ShutdownModule();
-		delete Module;
+		ShutdownModule();
 	}
 
 	void ParseSettingsFromCommandLine()
@@ -129,6 +166,11 @@ public:
 		HttpTestLogLevelInitializer.ResumeLogVerbosity();
 	}
 
+	TSharedRef<IHttpRequest> CreateRequest()
+	{
+		return bRetryEnabled ? HttpRetryManager->CreateRequest() : HttpModule->CreateRequest();
+	}
+
 	const FString UrlWithInvalidPortToTestConnectTimeout() const { return TEXT("http://10.255.255.1:8765"); } // non-routable IP address with a random port
 	const FString UrlBase() const { return FString::Format(TEXT("http://{0}:{1}"), { *WebServerIp, WebServerHttpPort }); }
 	const FString UrlHttpTests() const { return FString::Format(TEXT("{0}/webtests/httptests"), { *UrlBase() }); }
@@ -140,10 +182,11 @@ public:
 
 	FString WebServerIp;
 	uint32 WebServerHttpPort;
-	FMockHttpModule* HttpModule;
+	FMockHttpModule* HttpModule = nullptr;
 	bool bRunHeavyTests;
 	bool bRetryEnabled;
 	FHttpTestLogLevelInitializer HttpTestLogLevelInitializer;
+	TSharedPtr<FMockRetryManager> HttpRetryManager;
 };
 
 TEST_CASE_METHOD(FHttpModuleTestFixture, "Shutdown http module without issue when there are ongoing upload http requests.", HTTP_TAG)
@@ -159,7 +202,7 @@ TEST_CASE_METHOD(FHttpModuleTestFixture, "Shutdown http module without issue whe
 	{
 		IHttpRequest* LeakingHttpRequest = FPlatformHttp::ConstructRequest(); // Leaking in purpose to make sure it's ok
 
-		TSharedRef<IHttpRequest> HttpRequest = HttpModule->CreateRequest();
+		TSharedRef<IHttpRequest> HttpRequest = CreateRequest();
 		HttpRequest->SetURL(UrlToTestMethods());
 		HttpRequest->SetVerb(TEXT("PUT"));
 		// TODO: Use some shared data, like cookie, openssl session etc.
@@ -173,22 +216,73 @@ TEST_CASE_METHOD(FHttpModuleTestFixture, "Shutdown http module without issue whe
 	HttpModule->GetHttpManager().Tick(0.0f);
 }
 
-class FMockRetryManager : public FHttpRetrySystem::FManager
+TEST_CASE_METHOD(FHttpModuleTestFixture, "Shutdown http module without issue when there are ongoing streaming http requests with timeout.", HTTP_TAG)
 {
-public:
-	using FHttpRetrySystem::FManager::FManager;
-	using FHttpRetrySystem::FManager::RequestList;
-	using FHttpRetrySystem::FManager::FHttpRetryRequestEntry;
-	using FHttpRetrySystem::FManager::RetryTimeoutRelativeSecondsDefault;
-	using FHttpRetrySystem::FManager::RetryLimitCountDefault;
-	using FHttpRetrySystem::FManager::RetryLimitCountForConnectionErrorDefault;
+	if (!bRunHeavyTests)
+	{
+		return;
+	}
 
-	bool IsEmpty()
+	// When use generator, it doesn't do the ctor and dtor of FHttpModuleTestFixture each time, so manually 
+	// shutdown and init here to shutdown module a lot of times
+	ShutdownModule();
+	InitModule();
+
+	int32 NumRequests = GENERATE(
+		1, 2, 3, 4, 5, 6, 7, 8, 9, 10,
+		11, 12, 13, 14, 15, 16, 17, 18, 19, 20,
+		21, 22, 23, 24, 25, 26, 27, 28, 29, 30,
+		31, 32, 33, 34, 35, 36, 37, 38, 39, 40,
+		41, 42, 43, 44, 45, 46, 47, 48, 49, 50,
+		51, 52, 53, 54, 55, 56, 57, 58, 59, 60,
+		61, 62, 63, 64, 65, 66, 67, 68, 69, 70,
+		71, 72, 73, 74, 75, 76, 77, 78, 79, 80,
+		81, 82, 83, 84, 85, 86, 87, 88, 89, 90,
+		91, 92, 93, 94, 95, 96, 97, 98, 99, 100
+	);
+
+	//Output NumRequests when error occurs.
+	UNSCOPED_INFO(NumRequests);
+	HttpModule->HttpTotalTimeout = 2.0f;
+	HttpModule->HttpActivityTimeout = 1.0f;
+
+	DYNAMIC_SECTION(" making " << NumRequests << " requests")
 	{
-		FScopeLock ScopeLock(&RequestListLock);
-		return RequestList.IsEmpty();
+		DisableWarningsInThisTest();
+
+		uint32 ChunkSize = 1024 * 1024;
+		TArray<uint8> DataChunk;
+		DataChunk.SetNum(ChunkSize);
+		FMemory::Memset(DataChunk.GetData(), 'd', ChunkSize);
+
+		for (int32 i = 0; i < NumRequests; ++i)
+		{
+			{
+				TSharedRef<IHttpRequest> HttpRequest = CreateRequest();
+				HttpRequest->SetURL(UrlToTestMethods());
+				HttpRequest->SetVerb(TEXT("PUT"));
+				HttpRequest->SetContent(DataChunk);
+				HttpRequest->OnProcessRequestComplete().BindLambda([](FHttpRequestPtr HttpRequest, FHttpResponsePtr HttpResponse, bool bSucceeded) {
+					CHECK(bSucceeded);
+					});
+				HttpRequest->ProcessRequest();
+			}
+
+			{
+				TSharedRef<IHttpRequest> HttpRequest = CreateRequest();
+				HttpRequest->SetURL(UrlStreamDownload(2, HTTP_TEST_TIMEOUT_CHUNK_SIZE, 2));
+				HttpRequest->OnProcessRequestComplete().BindLambda([](FHttpRequestPtr HttpRequest, FHttpResponsePtr HttpResponse, bool bSucceeded) {
+					CHECK(bSucceeded);
+				});
+				HttpRequest->ProcessRequest();
+			}
+		}
+
+		HttpModule->GetHttpManager().Tick(0.0f);
 	}
-};
+
+	ShutdownModule();
+}
 
 class FWaitUntilCompleteHttpFixture : public FHttpModuleTestFixture
 {
@@ -197,11 +291,6 @@ public:
 	{
 		HttpModule->GetHttpManager().SetRequestAddedDelegate(FHttpManagerRequestAddedDelegate::CreateRaw(this, &FWaitUntilCompleteHttpFixture::OnRequestAdded));
 		HttpModule->GetHttpManager().SetRequestCompletedDelegate(FHttpManagerRequestCompletedDelegate::CreateRaw(this, &FWaitUntilCompleteHttpFixture::OnRequestCompleted));
-
-		if (bRetryEnabled)
-		{
-			HttpRetryManager = MakeShared<FMockRetryManager>(FHttpRetrySystem::FRetryLimitCountSetting(RetryLimitCount), FHttpRetrySystem::FRetryTimeoutRelativeSecondsSetting(/*RetryTimeoutRelativeSeconds*/));
-		}
 	}
 
 	~FWaitUntilCompleteHttpFixture()
@@ -251,16 +340,10 @@ public:
 		return OngoingRequests != 0 || (bRetryEnabled && !HttpRetryManager->IsEmpty());
 	}
 
-	TSharedRef<IHttpRequest> CreateRequest()
-	{
-		return bRetryEnabled ? HttpRetryManager->CreateRequest() : HttpModule->CreateRequest();
-	}
-
 	std::atomic<int32> OngoingRequests = 0;
 	float TickFrequency = 1.0f / 60; /*60 FPS*/;
 
 	uint32 RetryLimitCount = 0;
-	TSharedPtr<FMockRetryManager> HttpRetryManager;
 	uint32 ExpectingExtraCallbacks = 0;
 };
 
@@ -842,8 +925,6 @@ void WriteTestFile(const FString& TestFileName, uint64 TestFileSize)
 }
 }
 
-#define HTTP_TEST_TIMEOUT_CHUNK_SIZE 16*1024 // Use a big chunk size so it triggers data received callback in time on all platforms
-
 TEST_CASE_METHOD(FWaitUntilCompleteHttpFixture, "Http request activity timeout", HTTP_TAG)
 {
 	DisableWarningsInThisTest();
diff --git a/Engine/Source/Runtime/Online/HTTP/Private/HttpManager.cpp b/Engine/Source/Runtime/Online/HTTP/Private/HttpManager.cpp
--- a/Engine/Source/Runtime/Online/HTTP/Private/HttpManager.cpp
+++ b/Engine/Source/Runtime/Online/HTTP/Private/HttpManager.cpp
@@ -361,9 +361,14 @@ void FHttpManager::Flush(EHttpFlushReason FlushReason)
 			UE_CLOG(ShouldOutputHttpWarnings(), LogHttp, Warning, TEXT("[FHttpManager::Flush] FlushTimeSoftLimitSeconds [%.3fs] exceeded. Cancelling %d outstanding HTTP requests:"), FlushTimeSoftLimitSeconds, RequestsNum);
 
 			{
-				FScopeLock ScopeLock(&RequestLock);
+				TArray<FHttpRequestRef> RequestsToCancel;
+
+				{
+					FScopeLock ScopeLock(&RequestLock);
+					RequestsToCancel = Requests;
+				}
 
-				for (TArray<FHttpRequestRef>::TIterator It(Requests); It; ++It)
+				for (TArray<FHttpRequestRef>::TIterator It(RequestsToCancel); It; ++It)
 				{
 					FHttpRequestRef& Request = *It;
 
