Add the interface to set activity timeout per request.

#tests Passed WebTests test case

diff --git a/Engine/Source/Programs/WebTests/Private/TestHttp.cpp b/Engine/Source/Programs/WebTests/Private/TestHttp.cpp
--- a/Engine/Source/Programs/WebTests/Private/TestHttp.cpp
+++ b/Engine/Source/Programs/WebTests/Private/TestHttp.cpp
@@ -929,16 +929,26 @@ TEST_CASE_METHOD(FWaitUntilCompleteHttpFixture, "Http request activity timeout",
 {
 	DisableWarningsInThisTest();
 
-	float ReceiveTimeoutSetting = 3.0f;
-	HttpModule->HttpActivityTimeout = ReceiveTimeoutSetting;
+	float ActivityTimeoutSetting = 3.0f;
+	HttpModule->HttpActivityTimeout = ActivityTimeoutSetting;
 
 	TSharedPtr<IHttpRequest> HttpRequest = CreateRequest();
+
+	SECTION("By default activity timeout from http module")
+	{
+	}
+	SECTION("By customized activity timeout per http request which will override default settings from http module")
+	{
+		ActivityTimeoutSetting = 4.0f;
+		HttpRequest->SetActivityTimeout(ActivityTimeoutSetting);
+	}
+
 	HttpRequest->SetURL(UrlStreamDownload(3/*Chunks*/, HTTP_TEST_TIMEOUT_CHUNK_SIZE, 5/*ChunkLatency*/));
 	HttpRequest->SetVerb(TEXT("GET"));
 
 	const double StartTime = FPlatformTime::Seconds();
 
-	HttpRequest->OnProcessRequestComplete().BindLambda([StartTime, ReceiveTimeoutSetting](FHttpRequestPtr HttpRequest, FHttpResponsePtr HttpResponse, bool bSucceeded) {
+	HttpRequest->OnProcessRequestComplete().BindLambda([StartTime, ActivityTimeoutSetting](FHttpRequestPtr HttpRequest, FHttpResponsePtr HttpResponse, bool bSucceeded) {
 		CHECK(!bSucceeded);
 		CHECK(HttpRequest->GetStatus() == EHttpRequestStatus::Failed);
 		CHECK(HttpRequest->GetFailureReason() == EHttpFailureReason::ConnectionError);
@@ -950,9 +960,9 @@ TEST_CASE_METHOD(FWaitUntilCompleteHttpFixture, "Http request activity timeout",
 		// In a similar way on MacOS/iOS we don't get any notification until some data is received
 		// So it takes 5s to receive the first chunk to be considered as connected, then start response timer and
 		// take 3s to response timeout
-		CHECK(FMath::IsNearlyEqual(DurationInSeconds, ReceiveTimeoutSetting + 5, HTTP_TIME_DIFF_TOLERANCE_OF_REQUEST));
+		CHECK(FMath::IsNearlyEqual(DurationInSeconds, ActivityTimeoutSetting + 5, HTTP_TIME_DIFF_TOLERANCE_OF_REQUEST));
 #else
-		CHECK(FMath::IsNearlyEqual(DurationInSeconds, ReceiveTimeoutSetting, HTTP_TIME_DIFF_TOLERANCE_OF_REQUEST));
+		CHECK(FMath::IsNearlyEqual(DurationInSeconds, ActivityTimeoutSetting, HTTP_TIME_DIFF_TOLERANCE_OF_REQUEST));
 #endif
 
 	});
diff --git a/Engine/Source/Runtime/Online/BuildPatchServices/Private/Tests/Mock/HttpRequest.mock.h b/Engine/Source/Runtime/Online/BuildPatchServices/Private/Tests/Mock/HttpRequest.mock.h
--- a/Engine/Source/Runtime/Online/BuildPatchServices/Private/Tests/Mock/HttpRequest.mock.h
+++ b/Engine/Source/Runtime/Online/BuildPatchServices/Private/Tests/Mock/HttpRequest.mock.h
@@ -140,6 +140,11 @@ namespace BuildPatchServices
 			return TOptional<float>();
 		}
 
+		virtual void SetActivityTimeout(float InTimeoutSecs) override
+		{
+			MOCK_FUNC_NOT_IMPLEMENTED("FMockHttpRequest::SetActivityTimeout");
+		}
+
 		virtual void ProcessRequestUntilComplete() override
 		{
 		}
diff --git a/Engine/Source/Runtime/Online/HTTP/Private/GenericPlatform/GenericPlatformHttp.cpp b/Engine/Source/Runtime/Online/HTTP/Private/GenericPlatform/GenericPlatformHttp.cpp
--- a/Engine/Source/Runtime/Online/HTTP/Private/GenericPlatform/GenericPlatformHttp.cpp
+++ b/Engine/Source/Runtime/Online/HTTP/Private/GenericPlatform/GenericPlatformHttp.cpp
@@ -128,6 +128,7 @@ public:
 	virtual void SetTimeout(float InTimeoutSecs) override {}
 	virtual void ClearTimeout() override {}
 	virtual TOptional<float> GetTimeout() const override { return TOptional<float>(); }
+	virtual void SetActivityTimeout(float InTimeoutSecs) override {}
 	virtual void ProcessRequestUntilComplete() override {}
 
 private:
diff --git a/Engine/Source/Runtime/Online/HTTP/Private/GenericPlatform/HttpRequestCommon.cpp b/Engine/Source/Runtime/Online/HTTP/Private/GenericPlatform/HttpRequestCommon.cpp
--- a/Engine/Source/Runtime/Online/HTTP/Private/GenericPlatform/HttpRequestCommon.cpp
+++ b/Engine/Source/Runtime/Online/HTTP/Private/GenericPlatform/HttpRequestCommon.cpp
@@ -205,6 +205,11 @@ float FHttpRequestCommon::GetTimeoutOrDefault() const
 	return GetTimeout().Get(FHttpModule::Get().GetHttpTotalTimeout());
 }
 
+void FHttpRequestCommon::SetActivityTimeout(float InTimeoutSecs)
+{
+	ActivityTimeoutSecs = InTimeoutSecs;
+}
+
 const FHttpResponsePtr FHttpRequestCommon::GetResponse() const
 {
 	return ResponseCommon;
@@ -254,7 +259,7 @@ void FHttpRequestCommon::StartActivityTimeoutTimer()
 		return;
 	}
 
-	float HttpActivityTimeout = FHttpModule::Get().GetHttpActivityTimeout();
+	float HttpActivityTimeout = GetActivityTimeoutOrDefault();
 	if (HttpActivityTimeout == 0)
 	{
 		return;
@@ -306,7 +311,7 @@ void FHttpRequestCommon::OnActivityTimeoutTimerTaskTrigger()
 	QUICK_SCOPE_CYCLE_COUNTER(STAT_FHttpRequestCommon_AbortRequest);
 	bActivityTimedOut = true;
 	AbortRequest();
-	UE_LOG(LogHttp, Log, TEXT("Request [%s] timed out at [%s] because of no responding for %0.2f seconds"), *GetURL(), *FDateTime::Now().ToString(TEXT("%H:%M:%S:%s")), FHttpModule::Get().GetHttpActivityTimeout());
+	UE_LOG(LogHttp, Log, TEXT("Request [%s] timed out at [%s] because of no responding for %0.2f seconds"), *GetURL(), *FDateTime::Now().ToString(TEXT("%H:%M:%S:%s")), GetActivityTimeoutOrDefault());
 }
 
 void FHttpRequestCommon::ResetActivityTimeoutTimer(FStringView Reason)
@@ -323,7 +328,7 @@ void FHttpRequestCommon::ResetActivityTimeoutTimer(FStringView Reason)
 		return;
 	}
 
-	ActivityTimeoutAt = FPlatformTime::Seconds() + FHttpModule::Get().GetHttpActivityTimeout();
+	ActivityTimeoutAt = FPlatformTime::Seconds() + GetActivityTimeoutOrDefault();
 	UE_LOG(LogHttp, VeryVerbose, TEXT("Request [%p] reset response timeout timer at %s: %s"), this, *FDateTime::Now().ToString(TEXT("%H:%M:%S:%s")), Reason.GetData());
 }
 
@@ -491,3 +496,9 @@ void FHttpRequestCommon::StopPassingReceivedData()
 
 	ResponseBodyReceiveStream = nullptr;
 }
+
+float FHttpRequestCommon::GetActivityTimeoutOrDefault() const
+{
+	return ActivityTimeoutSecs.Get(FHttpModule::Get().GetHttpActivityTimeout());
+
+}
\ No newline at end of file
diff --git a/Engine/Source/Runtime/Online/HTTP/Private/HttpRequestAdapter.cpp b/Engine/Source/Runtime/Online/HTTP/Private/HttpRequestAdapter.cpp
--- a/Engine/Source/Runtime/Online/HTTP/Private/HttpRequestAdapter.cpp
+++ b/Engine/Source/Runtime/Online/HTTP/Private/HttpRequestAdapter.cpp
@@ -112,6 +112,11 @@ TOptional<float> FHttpRequestAdapterBase::GetTimeout() const
 	return HttpRequest->GetTimeout(); 
 }
 
+void FHttpRequestAdapterBase::SetActivityTimeout(float InTimeoutSecs)
+{
+	return HttpRequest->SetActivityTimeout(InTimeoutSecs);
+}
+
 void FHttpRequestAdapterBase::ProcessRequestUntilComplete() 
 { 
 	return HttpRequest->ProcessRequestUntilComplete(); 
diff --git a/Engine/Source/Runtime/Online/HTTP/Private/NullHttp.cpp b/Engine/Source/Runtime/Online/HTTP/Private/NullHttp.cpp
--- a/Engine/Source/Runtime/Online/HTTP/Private/NullHttp.cpp
+++ b/Engine/Source/Runtime/Online/HTTP/Private/NullHttp.cpp
@@ -173,7 +173,7 @@ void FNullHttpRequest::Tick(float DeltaSeconds)
 	if (CompletionStatus == EHttpRequestStatus::Processing)
 	{
 		ElapsedTime += DeltaSeconds;
-		const float HttpTimeout = GetTimeout().Get(FHttpModule::Get().GetHttpActivityTimeout());
+		const float HttpTimeout = GetTimeout().Get(FHttpModule::Get().GetHttpTotalTimeout());
 		if (HttpTimeout > 0 && ElapsedTime >= HttpTimeout)
 		{
 			UE_LOG(LogHttp, Warning, TEXT("Timeout processing Http request. %p"),
@@ -224,6 +224,10 @@ TOptional<float> FNullHttpRequest::GetTimeout() const
 	return TimeoutSecs; 
 }
 
+void FNullHttpRequest::SetActivityTimeout(float InTimeoutSecs)
+{
+}
+
 void FNullHttpRequest::ProcessRequestUntilComplete()
 {
 }
diff --git a/Engine/Source/Runtime/Online/HTTP/Private/NullHttp.h b/Engine/Source/Runtime/Online/HTTP/Private/NullHttp.h
--- a/Engine/Source/Runtime/Online/HTTP/Private/NullHttp.h
+++ b/Engine/Source/Runtime/Online/HTTP/Private/NullHttp.h
@@ -46,6 +46,7 @@ public:
 	virtual void SetTimeout(float InTimeoutSecs) override;
 	virtual void ClearTimeout() override;
 	virtual TOptional<float> GetTimeout() const override;
+	virtual void SetActivityTimeout(float InTimeoutSecs) override;
 	virtual void ProcessRequestUntilComplete() override;
 
 	FNullHttpRequest()
diff --git a/Engine/Source/Runtime/Online/HTTP/Public/GenericPlatform/HttpRequestCommon.h b/Engine/Source/Runtime/Online/HTTP/Public/GenericPlatform/HttpRequestCommon.h
--- a/Engine/Source/Runtime/Online/HTTP/Public/GenericPlatform/HttpRequestCommon.h
+++ b/Engine/Source/Runtime/Online/HTTP/Public/GenericPlatform/HttpRequestCommon.h
@@ -30,6 +30,8 @@ public:
 	HTTP_API virtual TOptional<float> GetTimeout() const override;
 	HTTP_API float GetTimeoutOrDefault() const;
 
+	HTTP_API virtual void SetActivityTimeout(float InTimeoutSecs) override;
+
 	HTTP_API virtual const FHttpResponsePtr GetResponse() const override;
 
 	// Can be called on game thread or http thread depend on the delegate thread policy
@@ -83,6 +85,8 @@ protected:
 	HTTP_API bool PassReceivedDataToStream(void* Ptr, int64 Length);
 	HTTP_API void StopPassingReceivedData();
 
+	HTTP_API float GetActivityTimeoutOrDefault() const;
+
 protected:
 	/** Current status of request being processed */
 	EHttpRequestStatus::Type CompletionStatus = EHttpRequestStatus::NotStarted;
@@ -96,6 +100,9 @@ protected:
 	/** Timeout in seconds for the entire HTTP request to complete */
 	TOptional<float> TimeoutSecs;
 
+	/** Timeout in seconds for the HTTP request activity timeout */
+	TOptional<float> ActivityTimeoutSecs;
+
 	/** Indicate the request is timed out, it should quit and fail with EHttpFailureReason::TimedOut */
 	std::atomic<bool> bTimedOut = false;
 	/** Indicate the request is activity timed out, it should quit and fail with EHttpFailureReason::ConnectionError */
diff --git a/Engine/Source/Runtime/Online/HTTP/Public/HttpRequestAdapter.h b/Engine/Source/Runtime/Online/HTTP/Public/HttpRequestAdapter.h
--- a/Engine/Source/Runtime/Online/HTTP/Public/HttpRequestAdapter.h
+++ b/Engine/Source/Runtime/Online/HTTP/Public/HttpRequestAdapter.h
@@ -37,6 +37,7 @@ public:
 	HTTP_API virtual void SetTimeout(float InTimeoutSecs) override;
 	HTTP_API virtual void ClearTimeout() override;
 	HTTP_API virtual TOptional<float> GetTimeout() const override;
+	HTTP_API virtual void SetActivityTimeout(float InTimeoutSecs) override;
 	HTTP_API virtual void ProcessRequestUntilComplete() override;
 	HTTP_API virtual const FHttpResponsePtr GetResponse() const override;
 	HTTP_API virtual float GetElapsedTime() const override;
diff --git a/Engine/Source/Runtime/Online/HTTP/Public/Interfaces/IHttpRequest.h b/Engine/Source/Runtime/Online/HTTP/Public/Interfaces/IHttpRequest.h
--- a/Engine/Source/Runtime/Online/HTTP/Public/Interfaces/IHttpRequest.h
+++ b/Engine/Source/Runtime/Online/HTTP/Public/Interfaces/IHttpRequest.h
@@ -268,6 +268,16 @@ public:
 	 */
 	virtual void SetTimeout(float InTimeoutSecs) = 0;
 
+	/**
+	 * Sets an optional activity timeout in seconds for this HTTP request. After connecting to 
+	 * web server, if there is no activity(send or receive) happen for this time period, it will
+	 * trigger activity timeout
+	 * If set, this value overrides the default HTTP activity timeout
+	 *
+	 * @param InTimeoutSecs - Timeout for this HTTP request instance, in seconds
+	 */
+	virtual void SetActivityTimeout(float InTimeoutSecs) = 0;
+
 	/**
 	 * Clears the optional timeout in seconds for this HTTP request, causing the default value
 	 * from FHttpModule::GetTimeout() to be used.
