Fix the test case "Http request connect timeout" for xCurl by increasing UE_HTTP_CONNECTION_TIMEOUT_MAX_DEVIATION;
Filter out the test case "Cancel http request connect before timeout" for platforms which doesn't support connect timeout retry with precise timeout value.

diff --git a/Engine/Source/Programs/WebTests/Private/TestHttp.cpp b/Engine/Source/Programs/WebTests/Private/TestHttp.cpp
--- a/Engine/Source/Programs/WebTests/Private/TestHttp.cpp
+++ b/Engine/Source/Programs/WebTests/Private/TestHttp.cpp
@@ -1686,6 +1686,7 @@ TEST_CASE_METHOD(FWaitUntilQuitFromTestThreadedFixture, "Threaded http request p
 	ThreadedHttpRunnable.StartTestHttpThread(false/*bBlockGameThread*/);
 }
 
+#if UE_HTTP_CONNECTION_TIMEOUT_SUPPORT_RETRY
 TEST_CASE_METHOD(FWaitUntilCompleteHttpFixture, "Cancel http request connect before timeout", HTTP_TAG)
 {
 	DisableWarningsInThisTest();
@@ -1714,6 +1715,7 @@ TEST_CASE_METHOD(FWaitUntilCompleteHttpFixture, "Cancel http request connect bef
 	HttpRequest->CancelRequest();
 	HttpRequest->CancelRequest(); // Duplicated calls to CancelRequest should be fine
 }
+#endif
 
 TEST_CASE_METHOD(FWaitUntilCompleteHttpFixture, "Retry respect Retry-After header in response", HTTP_TAG)
 {
diff --git a/Engine/Source/Runtime/Online/HTTP/HTTP.Build.cs b/Engine/Source/Runtime/Online/HTTP/HTTP.Build.cs
--- a/Engine/Source/Runtime/Online/HTTP/HTTP.Build.cs
+++ b/Engine/Source/Runtime/Online/HTTP/HTTP.Build.cs
@@ -136,7 +136,7 @@ public class HTTP : ModuleRules
 		float PlatformConnectionTimeoutMaxDeviation = 0.5f;
 		if (bPlatformSupportsXCurl)
 		{
-			PlatformConnectionTimeoutMaxDeviation = 3.5f;
+			PlatformConnectionTimeoutMaxDeviation = 4.0f;
 		}
 		if (Target.IsInPlatformGroup(UnrealPlatformGroup.Apple))
 		{
