Fix the test case crash "Shutdown http module without issue when there are ongoing download http requests" in WebTests.

Reason: UserStreamingInstance instance got deleted before http module shutdown, the data received delegates only unbind when http module shutdown.

diff --git a/Engine/Source/Programs/WebTests/Private/TestHttp.cpp b/Engine/Source/Programs/WebTests/Private/TestHttp.cpp
--- a/Engine/Source/Programs/WebTests/Private/TestHttp.cpp
+++ b/Engine/Source/Programs/WebTests/Private/TestHttp.cpp
@@ -703,16 +703,10 @@ class FInvalidateDelegateShutdownFixture : public FHttpModuleTestFixture
 public:
 	FInvalidateDelegateShutdownFixture()
 	{
-		UserStreamingInstance = new FUserStreamingClass;
+		UserStreamingInstance = MakeShared<FUserStreamingClass>();
 	}
 
-	~FInvalidateDelegateShutdownFixture()
-	{
-		delete UserStreamingInstance;
-		UserStreamingInstance = nullptr;
-	}
-
-	FUserStreamingClass* UserStreamingInstance;
+	TSharedPtr<FUserStreamingClass> UserStreamingInstance;
 };
 
 TEST_CASE_METHOD(FInvalidateDelegateShutdownFixture, "Shutdown http module without issue when there are ongoing download http requests", HTTP_TAG)
@@ -724,7 +718,7 @@ TEST_CASE_METHOD(FInvalidateDelegateShutdownFixture, "Shutdown http module witho
 		TSharedRef<IHttpRequest> HttpRequest = HttpModule->CreateRequest();
 		HttpRequest->SetURL(UrlStreamDownload(10, 1024*1024));
 		FHttpRequestStreamDelegate Delegate;
-		Delegate.BindRaw(UserStreamingInstance, &FUserStreamingClass::OnReceivedData);
+		Delegate.BindThreadSafeSP(UserStreamingInstance.ToSharedRef(), &FUserStreamingClass::OnReceivedData);
 		CHECK(HttpRequest->SetResponseBodyReceiveStreamDelegate(Delegate));
 
 		HttpRequest->OnProcessRequestComplete().BindLambda([](FHttpRequestPtr HttpRequest, FHttpResponsePtr HttpResponse, bool bSucceeded) {
