Fix the potential http crash when shared ptr bound to data receiving delegate got destroyed.

#tests Validated through test cases in WebTests project

diff --git a/Engine/Source/Programs/WebTests/Private/TestHttp.cpp b/Engine/Source/Programs/WebTests/Private/TestHttp.cpp
--- a/Engine/Source/Programs/WebTests/Private/TestHttp.cpp
+++ b/Engine/Source/Programs/WebTests/Private/TestHttp.cpp
@@ -698,6 +698,29 @@ TEST_CASE_METHOD(FWaitUntilCompleteHttpFixture, "In streaming downloading http r
 	UserInstance.Reset();
 }
 
+TEST_CASE_METHOD(FWaitUntilCompleteHttpFixture, "In streaming downloading http request won't crash if shared ptr bound to delegate got destroyed", HTTP_TAG)
+{
+	DisableWarningsInThisTest(); // Failed writing received data to disk/application
+
+	TSharedRef<IHttpRequest> HttpRequest = CreateRequest();
+	HttpRequest->SetURL(UrlStreamDownload(30, 1024*1024));
+
+	TSharedPtr<FUserStreamingClass> UserInstance = MakeShared<FUserStreamingClass>();
+
+	FHttpRequestStreamDelegate Delegate;
+	Delegate.BindThreadSafeSP(UserInstance.ToSharedRef(), &FUserStreamingClass::OnReceivedData);
+	CHECK(HttpRequest->SetResponseBodyReceiveStreamDelegate(Delegate));
+	HttpRequest->ProcessRequest();
+
+	while (*UserInstance->TotalBytesReceived == 0) // Make sure it started receiving data
+	{
+		FPlatformProcess::Sleep(0.001f);
+	}
+	CHECK(*UserInstance->TotalBytesReceived < 60 * 1024 * 1024);
+	CHECK(UserInstance.GetSharedReferenceCount() == 1);
+	UserInstance.Reset();
+}
+
 class FInvalidateDelegateShutdownFixture : public FHttpModuleTestFixture
 {
 public:
diff --git a/Engine/Source/Runtime/Online/HTTP/Public/Interfaces/IHttpRequest.h b/Engine/Source/Runtime/Online/HTTP/Public/Interfaces/IHttpRequest.h
--- a/Engine/Source/Runtime/Online/HTTP/Public/Interfaces/IHttpRequest.h
+++ b/Engine/Source/Runtime/Online/HTTP/Public/Interfaces/IHttpRequest.h
@@ -107,7 +107,7 @@ public:
 
 	virtual void Serialize(void* V, int64 Length) override
 	{
-		if (!StreamDelegate.Execute(V, Length))
+		if (!StreamDelegate.IsBound() || !StreamDelegate.Execute(V, Length))
 		{
 			SetError();
 		}
@@ -200,8 +200,6 @@ public:
 	 *     it's thread-safe in there
 	 *   - Make sure the delegate is safe to be called until receiving the process complete callback
 	 *     or after canceling the request
-	 *     For example: don't destroy the instance even if using BindThreadSafeSP, because internally 
-	 *     it's calling Execute to handle error by returned value instead of calling ExecuteIfBound
 	 * @param StreamDelegate - delegate from which the payload should be streamed.
 	 * @return True if the delegate can be used to stream the request. False otherwise.
 	 */
