Use safer call ExecuteIfBound for http response stream delegate, to avoid potential thread data race issue.

diff --git a/Engine/Plugins/Media/ElectraUtil/Source/ElectraHTTPStream/Private/Generic/GenericElectraHTTPStream.cpp b/Engine/Plugins/Media/ElectraUtil/Source/ElectraHTTPStream/Private/Generic/GenericElectraHTTPStream.cpp
--- a/Engine/Plugins/Media/ElectraUtil/Source/ElectraHTTPStream/Private/Generic/GenericElectraHTTPStream.cpp
+++ b/Engine/Plugins/Media/ElectraUtil/Source/ElectraHTTPStream/Private/Generic/GenericElectraHTTPStream.cpp
@@ -6,6 +6,7 @@
 
 #include "Http.h"
 
+#include "HAL/IConsoleManager.h"
 #include "HAL/ThreadSafeCounter.h"
 #include "HAL/PlatformTime.h"
 #include "HAL/PlatformProcess.h"
@@ -34,6 +35,15 @@ DECLARE_STATS_GROUP(TEXT("Electra HTTP Stream"), STATGROUP_ElectraHTTPStream, ST
 DECLARE_CYCLE_STAT(TEXT("Process"), STAT_ElectraHTTPThread_Process, STATGROUP_ElectraHTTPStream);
 DECLARE_CYCLE_STAT(TEXT("Custom handler"), STAT_ElectraHTTPThread_CustomHandler, STATGROUP_ElectraHTTPStream);
 
+/***************************************************************************************************************************************************/
+/***************************************************************************************************************************************************/
+/***************************************************************************************************************************************************/
+TAutoConsoleVariable<bool> CVarElectraUseHttpStreamDelegateV2(
+	TEXT("electra.UseHttpStreamDelegateV2"),
+	true,
+	TEXT("Use http stream delegate v2 in electra")
+);
+
 /***************************************************************************************************************************************************/
 /***************************************************************************************************************************************************/
 /***************************************************************************************************************************************************/
@@ -244,6 +254,7 @@ private:
 	void OnHeaderReceived(FHttpRequestPtr InSourceHttpRequest, const FString& InHeaderName, const FString& InHeaderValue);
 	void OnStatusCodeReceived(FHttpRequestPtr InSourceHttpRequest, int32 InHttpStatusCode);
 	bool OnProcessRequestStream(void *InDataPtr, int64 InLength);
+	void OnProcessRequestStreamV2(void *InDataPtr, int64& InLength);
 
 	// Owner to be notified on activity.
 	TWeakPtr<FElectraHTTPStreamGeneric, ESPMode::ThreadSafe> Owner;
@@ -325,10 +336,22 @@ bool FElectraHTTPStreamRequestGeneric::Setup()
 	RequestHandle->OnHeaderReceived().BindThreadSafeSP(AsShared(), &FElectraHTTPStreamRequestGeneric::OnHeaderReceived);
 	RequestHandle->OnStatusCodeReceived().BindThreadSafeSP(AsShared(), &FElectraHTTPStreamRequestGeneric::OnStatusCodeReceived);
 
-	FHttpRequestStreamDelegate StreamDelegate;
-	StreamDelegate.BindThreadSafeSP(AsShared(), &FElectraHTTPStreamRequestGeneric::OnProcessRequestStream);
-	bool bOk = RequestHandle->SetResponseBodyReceiveStreamDelegate(StreamDelegate);
-	(void)bOk; check(bOk);
+	if (CVarElectraUseHttpStreamDelegateV2.GetValueOnAnyThread())
+	{
+		FHttpRequestStreamDelegateV2 StreamDelegate;
+		StreamDelegate.BindThreadSafeSP(AsShared(), &FElectraHTTPStreamRequestGeneric::OnProcessRequestStreamV2);
+		bool bOk = RequestHandle->SetResponseBodyReceiveStreamDelegateV2(StreamDelegate);
+		(void)bOk; check(bOk);
+	}
+	else
+	{
+PRAGMA_DISABLE_DEPRECATION_WARNINGS 
+		FHttpRequestStreamDelegate StreamDelegate;
+		StreamDelegate.BindThreadSafeSP(AsShared(), &FElectraHTTPStreamRequestGeneric::OnProcessRequestStream);
+		bool bOk = RequestHandle->SetResponseBodyReceiveStreamDelegate(StreamDelegate);
+		(void)bOk; check(bOk);
+PRAGMA_ENABLE_DEPRECATION_WARNINGS 
+	}
 
 	// We set the user agent manually. For simplicities sake we add it to the list of additional headers.
 	AdditionalHeaders.FindOrAdd(TEXT("User-Agent"), UserAgent);
@@ -631,6 +654,45 @@ bool FElectraHTTPStreamRequestGeneric::OnProcessRequestStream(void *InDataPtr, i
 	return true;
 }
 
+void FElectraHTTPStreamRequestGeneric::OnProcessRequestStreamV2(void *InDataPtr, int64& InOutLength)
+{
+	if (InDataPtr == nullptr || InOutLength < 0)
+	{
+		InOutLength = 0;
+		return;
+	}
+	if (WasCanceled())
+	{
+		return;
+	}
+	double Now = FPlatformTime::Seconds();
+	if (CurrentState < EState::ReadingResponseData)
+	{
+		CurrentState = EState::ReadingResponseData;
+		Response->TimeUntilFirstByte = Now - Response->StartTime;
+		// Parse the headers and report them.
+		if (!ParseResponseHeaders())
+		{
+			InOutLength = 0;
+			return;
+		}
+	}
+
+	// Add the data to the response.
+	TConstArrayView<const uint8> Data(static_cast<const uint8*>(InDataPtr), (int32)InOutLength);
+	Response->AddResponseData(Data);
+	// Notify amount of new data available.
+	NotifyCallback(EElectraHTTPStreamNotificationReason::ReadData, InOutLength);
+
+	Response->TimeOfMostRecentReceive = Now;
+
+	TSharedPtr<FElectraHTTPStreamGeneric, ESPMode::ThreadSafe> PinnedOwner = Owner.Pin();
+	if (PinnedOwner.IsValid())
+	{
+		PinnedOwner->TriggerWorkSignal();
+	}
+}
+
 
 /***************************************************************************************************************************************************/
 /***************************************************************************************************************************************************/
diff --git a/Engine/Source/Programs/WebTests/Private/TestHttp.cpp b/Engine/Source/Programs/WebTests/Private/TestHttp.cpp
--- a/Engine/Source/Programs/WebTests/Private/TestHttp.cpp
+++ b/Engine/Source/Programs/WebTests/Private/TestHttp.cpp
@@ -552,12 +552,11 @@ TEST_CASE_METHOD(FWaitUntilCompleteHttpFixture, "Streaming http download", HTTP_
 	}
 	SECTION("Success with customized stream delegate")
 	{
-		FHttpRequestStreamDelegate Delegate;
-		Delegate.BindLambda([TotalBytesReceived](void* Ptr, int64 Length) {
+		FHttpRequestStreamDelegateV2 Delegate;
+		Delegate.BindLambda([TotalBytesReceived](void* Ptr, int64& Length) {
 			*TotalBytesReceived += Length;
-			return true;
 		});
-		CHECK(HttpRequest->SetResponseBodyReceiveStreamDelegate(Delegate));
+		CHECK(HttpRequest->SetResponseBodyReceiveStreamDelegateV2(Delegate));
 
 		HttpRequest->OnProcessRequestComplete().BindLambda([Chunks, ChunkSize, TotalBytesReceived](FHttpRequestPtr HttpRequest, FHttpResponsePtr HttpResponse, bool bSucceeded) {
 			CHECK(bSucceeded);
@@ -603,12 +602,12 @@ TEST_CASE_METHOD(FWaitUntilCompleteHttpFixture, "Streaming http download", HTTP_
 	{
 		DisableWarningsInThisTest();
 
-		FHttpRequestStreamDelegate Delegate;
-		Delegate.BindLambda([TotalBytesReceived](void* Ptr, int64 Length) {
+		FHttpRequestStreamDelegateV2 Delegate;
+		Delegate.BindLambda([TotalBytesReceived](void* Ptr, int64& Length) {
 			*TotalBytesReceived += Length;
-			return false;
+			Length = 0; // Mark as no data was serialized successfully
 		});
-		CHECK(HttpRequest->SetResponseBodyReceiveStreamDelegate(Delegate));
+		CHECK(HttpRequest->SetResponseBodyReceiveStreamDelegateV2(Delegate));
 
 		HttpRequest->OnProcessRequestComplete().BindLambda([ChunkSize, TotalBytesReceived](FHttpRequestPtr HttpRequest, FHttpResponsePtr HttpResponse, bool bSucceeded) {
 			CHECK(!bSucceeded);
@@ -662,10 +661,9 @@ public:
 		TotalBytesReceived = nullptr;
 	}
 
-	bool OnReceivedData(void* Ptr, int64 Length)
+	void OnReceivedData(void* Ptr, int64& Length)
 	{
 		*TotalBytesReceived += Length;
-		return true;
 	}
 
 	int64* TotalBytesReceived;
@@ -678,9 +676,9 @@ TEST_CASE_METHOD(FWaitUntilCompleteHttpFixture, "In streaming downloading http r
 
 	TSharedPtr<FUserStreamingClass> UserInstance = MakeShared<FUserStreamingClass>();
 
-	FHttpRequestStreamDelegate Delegate;
+	FHttpRequestStreamDelegateV2 Delegate;
 	Delegate.BindThreadSafeSP(UserInstance.ToSharedRef(), &FUserStreamingClass::OnReceivedData);
-	CHECK(HttpRequest->SetResponseBodyReceiveStreamDelegate(Delegate));
+	CHECK(HttpRequest->SetResponseBodyReceiveStreamDelegateV2(Delegate));
 
 	HttpRequest->OnProcessRequestComplete().BindLambda([](FHttpRequestPtr HttpRequest, FHttpResponsePtr HttpResponse, bool bSucceeded) {
 		CHECK(!bSucceeded);
@@ -707,9 +705,9 @@ TEST_CASE_METHOD(FWaitUntilCompleteHttpFixture, "In streaming downloading http r
 
 	TSharedPtr<FUserStreamingClass> UserInstance = MakeShared<FUserStreamingClass>();
 
-	FHttpRequestStreamDelegate Delegate;
+	FHttpRequestStreamDelegateV2 Delegate;
 	Delegate.BindThreadSafeSP(UserInstance.ToSharedRef(), &FUserStreamingClass::OnReceivedData);
-	CHECK(HttpRequest->SetResponseBodyReceiveStreamDelegate(Delegate));
+	CHECK(HttpRequest->SetResponseBodyReceiveStreamDelegateV2(Delegate));
 	HttpRequest->ProcessRequest();
 
 	while (*UserInstance->TotalBytesReceived == 0) // Make sure it started receiving data
@@ -740,9 +738,9 @@ TEST_CASE_METHOD(FInvalidateDelegateShutdownFixture, "Shutdown http module witho
 	{
 		TSharedRef<IHttpRequest> HttpRequest = HttpModule->CreateRequest();
 		HttpRequest->SetURL(UrlStreamDownload(10, 1024*1024));
-		FHttpRequestStreamDelegate Delegate;
+		FHttpRequestStreamDelegateV2 Delegate;
 		Delegate.BindThreadSafeSP(UserStreamingInstance.ToSharedRef(), &FUserStreamingClass::OnReceivedData);
-		CHECK(HttpRequest->SetResponseBodyReceiveStreamDelegate(Delegate));
+		CHECK(HttpRequest->SetResponseBodyReceiveStreamDelegateV2(Delegate));
 
 		HttpRequest->OnProcessRequestComplete().BindLambda([](FHttpRequestPtr HttpRequest, FHttpResponsePtr HttpResponse, bool bSucceeded) {
 			CHECK(bSucceeded);
@@ -793,12 +791,11 @@ TEST_CASE_METHOD(FWaitUntilCompleteHttpFixture, "Can download big file exceeds 3
 	HttpRequest->SetVerb(TEXT("GET"));
 
 	TSharedRef<int64> TotalBytesReceived = MakeShared<int64>(0);
-	FHttpRequestStreamDelegate Delegate;
-	Delegate.BindLambda([TotalBytesReceived](void* Ptr, int64 Length) {
+	FHttpRequestStreamDelegateV2 Delegate;
+	Delegate.BindLambda([TotalBytesReceived](void* Ptr, int64& Length) {
 		*TotalBytesReceived += Length;
-		return true;
 	});
-	HttpRequest->SetResponseBodyReceiveStreamDelegate(Delegate);
+	HttpRequest->SetResponseBodyReceiveStreamDelegateV2(Delegate);
 
 	HttpRequest->OnProcessRequestComplete().BindLambda([Chunks, ChunkSize, TotalBytesReceived](FHttpRequestPtr HttpRequest, FHttpResponsePtr HttpResponse, bool bSucceeded) {
 		CHECK(bSucceeded);
@@ -1516,13 +1513,12 @@ TEST_CASE_METHOD(FValidateHeaderReceiveOrderFixture, "Http request header receiv
 	HttpRequest->SetURL(UrlStreamDownload(2/*Chunks*/, 1024/*ChunkSize*/));
 	HttpRequest->SetVerb(TEXT("GET"));
 
-	FHttpRequestStreamDelegate StreamDelegate;
-	StreamDelegate.BindLambda([this](void *InDataPtr, int64 InLength) {
+	FHttpRequestStreamDelegateV2 StreamDelegate;
+	StreamDelegate.BindLambda([this](void *InDataPtr, int64& InLength) {
 		bAnyDataReceived = true;
 		CHECK(!bCompleteCallbackTriggered);
-		return true;
 	});
-	HttpRequest->SetResponseBodyReceiveStreamDelegate(StreamDelegate);
+	HttpRequest->SetResponseBodyReceiveStreamDelegateV2(StreamDelegate);
 
 	SECTION("in http thread")
 	{
diff --git a/Engine/Source/Runtime/Online/HTTP/Private/Interfaces/IHttpRequest.cpp b/Engine/Source/Runtime/Online/HTTP/Private/Interfaces/IHttpRequest.cpp
new file mode 100644
--- /dev/null
+++ b/Engine/Source/Runtime/Online/HTTP/Private/Interfaces/IHttpRequest.cpp
@@ -0,0 +1,43 @@
+// Copyright Epic Games, Inc. All Rights Reserved.
+
+#include "Interfaces/IHttpRequest.h"
+
+class FArchiveWithDelegateV2 final : public FArchive
+{
+public:
+	FArchiveWithDelegateV2(FHttpRequestStreamDelegateV2 InStreamDelegateV2)
+		: StreamDelegateV2(InStreamDelegateV2)
+	{
+	}
+
+	virtual void Serialize(void* V, int64 Length) override
+	{
+		int64 LengthProcessed = Length;
+		StreamDelegateV2.ExecuteIfBound(V, LengthProcessed);
+		if (LengthProcessed != Length)
+		{
+			SetError();
+		}
+	}
+
+private:
+	FHttpRequestStreamDelegateV2 StreamDelegateV2;
+};
+
+
+PRAGMA_DISABLE_DEPRECATION_WARNINGS 
+bool IHttpRequest::SetContentFromStreamDelegate(FHttpRequestStreamDelegate StreamDelegate) 
+{ 
+	return SetContentFromStream(MakeShared<FArchiveWithDelegate>(StreamDelegate)); 
+}
+
+bool IHttpRequest::SetResponseBodyReceiveStreamDelegate(FHttpRequestStreamDelegate StreamDelegate) 
+{ 
+	return SetResponseBodyReceiveStream(MakeShared<FArchiveWithDelegate>(StreamDelegate)); 
+}
+PRAGMA_ENABLE_DEPRECATION_WARNINGS 
+
+bool IHttpRequest::SetResponseBodyReceiveStreamDelegateV2(FHttpRequestStreamDelegateV2 StreamDelegate) 
+{ 
+	return SetResponseBodyReceiveStream(MakeShared<FArchiveWithDelegateV2>(StreamDelegate)); 
+}
diff --git a/Engine/Source/Runtime/Online/HTTP/Public/Interfaces/IHttpRequest.h b/Engine/Source/Runtime/Online/HTTP/Public/Interfaces/IHttpRequest.h
--- a/Engine/Source/Runtime/Online/HTTP/Public/Interfaces/IHttpRequest.h
+++ b/Engine/Source/Runtime/Online/HTTP/Public/Interfaces/IHttpRequest.h
@@ -94,10 +94,18 @@ using FHttpRequestWillRetryDelegate = TTSDelegate<void(FHttpRequestPtr /*Request
  */
 using FHttpRequestStreamDelegate = TTSDelegate<bool(void*/*Ptr*/, int64/*Length*/)>;
 
+/**
+ * Delegate called when an Http request will send/recv data through stream
+ *
+ * @param Ptr - The buffer ptr to read/write
+ * @param InOutLength - The int64 reference length of buffer to read/write, if there is any error when serialize set it to 0
+ */
+using FHttpRequestStreamDelegateV2 = TTSDelegate<void(void*/*Ptr*/, int64&/*InOutLength*/)>;
+
 /**
  * Delegate version of FArchive, for streaming interface
  */
-class FArchiveWithDelegate final : public FArchive
+class UE_DEPRECATED(5.5, "FArchiveWithDelegate is deprecated and will be moved to internal") FArchiveWithDelegate final : public FArchive
 {
 public:
 	FArchiveWithDelegate(FHttpRequestStreamDelegate InStreamDelegate)
@@ -203,7 +211,8 @@ public:
 	 * @param StreamDelegate - delegate from which the payload should be streamed.
 	 * @return True if the delegate can be used to stream the request. False otherwise.
 	 */
-	bool SetContentFromStreamDelegate(FHttpRequestStreamDelegate StreamDelegate) { return SetContentFromStream(MakeShared<FArchiveWithDelegate>(StreamDelegate)); }
+	UE_DEPRECATED(5.5, "SetContentFromStreamDelegate has been deprecated and will not be supported because there is no seek support through delegate. Implement your own FArchive instead.")
+	bool SetContentFromStreamDelegate(FHttpRequestStreamDelegate StreamDelegate);
 
 	/**
 	 * Sets the stream to receive the response body. Make sure to handle the cleanup of stream when
@@ -230,7 +239,21 @@ public:
 	 * @param StreamDelegate - will be used to receive the response body
 	 * @return True if the delegate can be used. False otherwise.
 	 */
-	bool SetResponseBodyReceiveStreamDelegate(FHttpRequestStreamDelegate StreamDelegate) { return SetResponseBodyReceiveStream(MakeShared<FArchiveWithDelegate>(StreamDelegate)); }
+	UE_DEPRECATED(5.5, "SetResponseBodyReceiveStreamDelegate has been deprecated, use SetResponseBodyReceiveStreamDelegateV2 instead")
+	bool SetResponseBodyReceiveStreamDelegate(FHttpRequestStreamDelegate StreamDelegate);
+
+	/**
+	 * Sets the delegate to receive the response body. Make sure to handle the cleanup of received data when
+	 * failed to process the data(StreamDelegate return false), this http request will fail and quit.
+	 *
+	 * NOTE: Once set, the data will no longer be cached in response, IHttpResponse::GetContent() and
+	 * IHttpResponse::GetContentAsString() will return empty result. The delegate will be called from
+	 * another thread other than the game thread
+	 *
+	 * @param StreamDelegate - will be used to receive the response body
+	 * @return True if the delegate can be used. False otherwise.
+	 */
+	bool SetResponseBodyReceiveStreamDelegateV2(FHttpRequestStreamDelegateV2 StreamDelegate);
 
 	/**
 	 * Sets optional header info.
