Add ResetTimeoutStatus method to clear bTimedOut flag for http request to be reused when the total timeout was set.

#tests Tested through new test case in WebTests project.

diff --git a/Engine/Source/Programs/WebTests/Private/TestHttp.cpp b/Engine/Source/Programs/WebTests/Private/TestHttp.cpp
--- a/Engine/Source/Programs/WebTests/Private/TestHttp.cpp
+++ b/Engine/Source/Programs/WebTests/Private/TestHttp.cpp
@@ -1253,6 +1253,30 @@ TEST_CASE_METHOD(FWaitUntilQuitFromTestFixture, "Http request can be reused", HT
 	HttpRequest->ProcessRequest();
 }
 
+TEST_CASE_METHOD(FWaitUntilQuitFromTestFixture, "Http request can be reused when there is total timeout setting", HTTP_TAG)
+{
+	DisableWarningsInThisTest();
+
+	TSharedRef<IHttpRequest> HttpRequest = CreateRequest();
+	HttpRequest->SetURL(UrlMockLatency(3));
+	HttpRequest->SetTimeout(2);
+
+	HttpRequest->OnProcessRequestComplete().BindLambda([this](FHttpRequestPtr HttpRequest, FHttpResponsePtr HttpResponse, bool bSucceeded) {
+		CHECK(!bSucceeded);
+		CHECK(HttpRequest->GetFailureReason() == EHttpFailureReason::TimedOut);
+
+		HttpRequest->SetURL(UrlMockLatency(1));
+		HttpRequest->ResetTimeoutStatus(); // Must do this in order to restart timeout
+
+		HttpRequest->OnProcessRequestComplete().BindLambda([this](FHttpRequestPtr HttpRequest, FHttpResponsePtr HttpResponse, bool bSucceeded) {
+			CHECK(bSucceeded);
+			bQuitRequested = true;
+		});
+		HttpRequest->ProcessRequest();
+	});
+	HttpRequest->ProcessRequest();
+}
+
 #if UE_HTTP_CONNECTION_TIMEOUT_SUPPORT_RETRY
 TEST_CASE_METHOD(FWaitUntilQuitFromTestFixture, "Make sure connection time out can work well for 2nd same http request", HTTP_TAG)
 {
diff --git a/Engine/Source/Runtime/Online/BuildPatchServices/Private/Tests/Mock/HttpRequest.mock.h b/Engine/Source/Runtime/Online/BuildPatchServices/Private/Tests/Mock/HttpRequest.mock.h
--- a/Engine/Source/Runtime/Online/BuildPatchServices/Private/Tests/Mock/HttpRequest.mock.h
+++ b/Engine/Source/Runtime/Online/BuildPatchServices/Private/Tests/Mock/HttpRequest.mock.h
@@ -134,6 +134,11 @@ namespace BuildPatchServices
 			MOCK_FUNC_NOT_IMPLEMENTED("FMockHttpRequest::ClearTimeout");
 		}
 
+		virtual void ResetTimeoutStatus() override
+		{
+			MOCK_FUNC_NOT_IMPLEMENTED("FMockHttpRequest::ResetTimeoutStatus");
+		}
+
 		virtual TOptional<float> GetTimeout() const override
 		{
 			MOCK_FUNC_NOT_IMPLEMENTED("FMockHttpRequest::GetTimeout");
diff --git a/Engine/Source/Runtime/Online/HTTP/Private/GenericPlatform/GenericPlatformHttp.cpp b/Engine/Source/Runtime/Online/HTTP/Private/GenericPlatform/GenericPlatformHttp.cpp
--- a/Engine/Source/Runtime/Online/HTTP/Private/GenericPlatform/GenericPlatformHttp.cpp
+++ b/Engine/Source/Runtime/Online/HTTP/Private/GenericPlatform/GenericPlatformHttp.cpp
@@ -127,6 +127,7 @@ public:
 	virtual EHttpRequestDelegateThreadPolicy GetDelegateThreadPolicy() const override { return EHttpRequestDelegateThreadPolicy::CompleteOnGameThread; }
 	virtual void SetTimeout(float InTimeoutSecs) override {}
 	virtual void ClearTimeout() override {}
+	virtual void ResetTimeoutStatus() override {}
 	virtual TOptional<float> GetTimeout() const override { return TOptional<float>(); }
 	virtual void SetActivityTimeout(float InTimeoutSecs) override {}
 	virtual void ProcessRequestUntilComplete() override {}
diff --git a/Engine/Source/Runtime/Online/HTTP/Private/GenericPlatform/HttpRequestCommon.cpp b/Engine/Source/Runtime/Online/HTTP/Private/GenericPlatform/HttpRequestCommon.cpp
--- a/Engine/Source/Runtime/Online/HTTP/Private/GenericPlatform/HttpRequestCommon.cpp
+++ b/Engine/Source/Runtime/Online/HTTP/Private/GenericPlatform/HttpRequestCommon.cpp
@@ -210,7 +210,13 @@ void FHttpRequestCommon::SetTimeout(float InTimeoutSecs)
 void FHttpRequestCommon::ClearTimeout()
 {
 	TimeoutSecs.Reset();
+	ResetTimeoutStatus();
+}
+
+void FHttpRequestCommon::ResetTimeoutStatus()
+{
 	StopTotalTimeoutTimer();
+	bTimedOut = false;
 }
 
 TOptional<float> FHttpRequestCommon::GetTimeout() const
diff --git a/Engine/Source/Runtime/Online/HTTP/Private/HttpRequestAdapter.cpp b/Engine/Source/Runtime/Online/HTTP/Private/HttpRequestAdapter.cpp
--- a/Engine/Source/Runtime/Online/HTTP/Private/HttpRequestAdapter.cpp
+++ b/Engine/Source/Runtime/Online/HTTP/Private/HttpRequestAdapter.cpp
@@ -107,6 +107,11 @@ void FHttpRequestAdapterBase::ClearTimeout()
 	HttpRequest->ClearTimeout(); 
 }
 
+void FHttpRequestAdapterBase::ResetTimeoutStatus() 
+{ 
+	HttpRequest->ResetTimeoutStatus(); 
+}
+
 TOptional<float> FHttpRequestAdapterBase::GetTimeout() const 
 { 
 	return HttpRequest->GetTimeout(); 
diff --git a/Engine/Source/Runtime/Online/HTTP/Private/NullHttp.cpp b/Engine/Source/Runtime/Online/HTTP/Private/NullHttp.cpp
--- a/Engine/Source/Runtime/Online/HTTP/Private/NullHttp.cpp
+++ b/Engine/Source/Runtime/Online/HTTP/Private/NullHttp.cpp
@@ -219,6 +219,10 @@ void FNullHttpRequest::ClearTimeout()
 	TimeoutSecs.Reset();
 }
 
+void FNullHttpRequest::ResetTimeoutStatus() 
+{
+}
+
 TOptional<float> FNullHttpRequest::GetTimeout() const 
 { 
 	return TimeoutSecs; 
diff --git a/Engine/Source/Runtime/Online/HTTP/Private/NullHttp.h b/Engine/Source/Runtime/Online/HTTP/Private/NullHttp.h
--- a/Engine/Source/Runtime/Online/HTTP/Private/NullHttp.h
+++ b/Engine/Source/Runtime/Online/HTTP/Private/NullHttp.h
@@ -45,6 +45,7 @@ public:
 	virtual EHttpRequestDelegateThreadPolicy GetDelegateThreadPolicy() const override;
 	virtual void SetTimeout(float InTimeoutSecs) override;
 	virtual void ClearTimeout() override;
+	virtual void ResetTimeoutStatus() override;
 	virtual TOptional<float> GetTimeout() const override;
 	virtual void SetActivityTimeout(float InTimeoutSecs) override;
 	virtual void ProcessRequestUntilComplete() override;
diff --git a/Engine/Source/Runtime/Online/HTTP/Public/GenericPlatform/HttpRequestCommon.h b/Engine/Source/Runtime/Online/HTTP/Public/GenericPlatform/HttpRequestCommon.h
--- a/Engine/Source/Runtime/Online/HTTP/Public/GenericPlatform/HttpRequestCommon.h
+++ b/Engine/Source/Runtime/Online/HTTP/Public/GenericPlatform/HttpRequestCommon.h
@@ -28,6 +28,7 @@ public:
 
 	HTTP_API virtual void SetTimeout(float InTimeoutSecs) override;
 	HTTP_API virtual void ClearTimeout() override;
+	HTTP_API virtual void ResetTimeoutStatus() override;
 	HTTP_API virtual TOptional<float> GetTimeout() const override;
 	HTTP_API float GetTimeoutOrDefault() const;
 
diff --git a/Engine/Source/Runtime/Online/HTTP/Public/HttpRequestAdapter.h b/Engine/Source/Runtime/Online/HTTP/Public/HttpRequestAdapter.h
--- a/Engine/Source/Runtime/Online/HTTP/Public/HttpRequestAdapter.h
+++ b/Engine/Source/Runtime/Online/HTTP/Public/HttpRequestAdapter.h
@@ -36,6 +36,7 @@ public:
 	HTTP_API virtual void AppendToHeader(const FString& HeaderName, const FString& AdditionalHeaderValue) override;
 	HTTP_API virtual void SetTimeout(float InTimeoutSecs) override;
 	HTTP_API virtual void ClearTimeout() override;
+	HTTP_API virtual void ResetTimeoutStatus() override;
 	HTTP_API virtual TOptional<float> GetTimeout() const override;
 	HTTP_API virtual void SetActivityTimeout(float InTimeoutSecs) override;
 	HTTP_API virtual void ProcessRequestUntilComplete() override;
diff --git a/Engine/Source/Runtime/Online/HTTP/Public/Interfaces/IHttpRequest.h b/Engine/Source/Runtime/Online/HTTP/Public/Interfaces/IHttpRequest.h
--- a/Engine/Source/Runtime/Online/HTTP/Public/Interfaces/IHttpRequest.h
+++ b/Engine/Source/Runtime/Online/HTTP/Public/Interfaces/IHttpRequest.h
@@ -304,6 +304,11 @@ public:
 	 */
 	virtual void ClearTimeout() = 0;
 
+	/**
+	 * Reset the elapsed timeout duration and flag, after the request completed and need to be reused
+	 */
+	virtual void ResetTimeoutStatus() = 0;
+
 	/**
 	 * Gets the optional timeout in seconds for this entire HTTP request to complete.
 	 * If valid, this value overrides the default HTTP timeout set via FHttpModule::SetTimeout().
