Make sure http request on platforms can cancel request with proper reason;
Add test case for canceling request without calling ProcessRequest, make sure complete delegate will be invoked;
Add test case for canceling request which called ProcessRequest, but still in the queue, the complete delegate will be invoked immediately(Make sure canceled queued requests will complete immediately instead of waiting until they are popped from the queue).

diff --git a/Engine/Source/Programs/WebTests/Private/TestHttp.cpp b/Engine/Source/Programs/WebTests/Private/TestHttp.cpp
--- a/Engine/Source/Programs/WebTests/Private/TestHttp.cpp
+++ b/Engine/Source/Programs/WebTests/Private/TestHttp.cpp
@@ -38,6 +38,7 @@
 
 extern TAutoConsoleVariable<bool> CVarHttpInsecureProtocolEnabled;
 extern TAutoConsoleVariable<bool> CVarHttpRetrySystemNonGameThreadSupportEnabled;
+extern TAutoConsoleVariable<int32> CVarHttpMaxConcurrentRequests;
 
 class FMockHttpModule : public FHttpModule
 {
@@ -1723,6 +1724,64 @@ TEST_CASE_METHOD(FWaitUntilQuitFromTestThreadedFixture, "Threaded http request p
 	ThreadedHttpRunnable.StartTestHttpThread(false/*bBlockGameThread*/);
 }
 
+TEST_CASE_METHOD(FWaitUntilQuitFromTestFixture, "Cancel http request without ProcessRequest called", HTTP_TAG)
+{
+	TSharedRef<IHttpRequest> HttpRequest = CreateRequest();
+	HttpRequest->SetURL(UrlToTestMethods());
+	++ExpectingExtraCallbacks;
+	HttpRequest->OnProcessRequestComplete().BindLambda([this](FHttpRequestPtr HttpRequest, FHttpResponsePtr HttpResponse, bool bSucceeded) {
+		--ExpectingExtraCallbacks;
+		CHECK(!bSucceeded);
+		CHECK(HttpRequest->GetFailureReason() == EHttpFailureReason::Cancelled);
+		bQuitRequested = true;
+	});
+	HttpRequest->CancelRequest();
+}
+
+TEST_CASE_METHOD(FWaitThreadedHttpFixture, "Cancel http request with ProcessRequest called but before started from queue", HTTP_TAG)
+{
+	CVarHttpMaxConcurrentRequests->Set(1);
+
+	std::atomic<bool> bFirstRequestCompleted = false;
+
+	ThreadedHttpRunnable.OnRunFromThread().BindLambda([this, &bFirstRequestCompleted]() {
+		TSharedRef<IHttpRequest> HttpRequestRunning = CreateRequest();
+		HttpRequestRunning->SetURL(UrlStreamDownload(3/*Chunks*/, HTTP_TEST_TIMEOUT_CHUNK_SIZE, 1/*ChunkLatency*/));
+		HttpRequestRunning->SetDelegateThreadPolicy(EHttpRequestDelegateThreadPolicy::CompleteOnHttpThread);
+		HttpRequestRunning->OnProcessRequestComplete().BindLambda([this, &bFirstRequestCompleted](FHttpRequestPtr HttpRequestQueuing, FHttpResponsePtr HttpResponse, bool bSucceeded) {
+			bFirstRequestCompleted = true;
+			ThreadedHttpRunnable.UnblockGameThread();
+		});
+		HttpRequestRunning->ProcessRequest();
+
+
+		TSharedRef<IHttpRequest> HttpRequestQueuing = CreateRequest();
+		HttpRequestQueuing->SetURL(UrlStreamDownload(3/*Chunks*/, HTTP_TEST_TIMEOUT_CHUNK_SIZE, 1/*ChunkLatency*/));
+		HttpRequestQueuing->SetDelegateThreadPolicy(EHttpRequestDelegateThreadPolicy::CompleteOnHttpThread);
+		HttpRequestQueuing->OnHeaderReceived().BindLambda([this](FHttpRequestPtr Request, const FString& HeaderName, const FString& HeaderValue) {
+			// Should never be started
+			CHECK(false);
+		});
+		HttpRequestQueuing->OnRequestProgress64().BindLambda([this](FHttpRequestPtr Request, uint64 /*BytesSent*/, uint64 BytesReceived) {
+			// Should never be started
+			CHECK(false);
+		});
+
+		++ExpectingExtraCallbacks;
+		HttpRequestQueuing->OnProcessRequestComplete().BindLambda([this, &bFirstRequestCompleted](FHttpRequestPtr HttpRequestQueuing, FHttpResponsePtr HttpResponse, bool bSucceeded) {
+			--ExpectingExtraCallbacks;
+			CHECK(!bSucceeded);
+			CHECK(HttpRequestQueuing->GetFailureReason() == EHttpFailureReason::Cancelled);
+			CHECK(!bFirstRequestCompleted);
+		});
+		HttpRequestQueuing->ProcessRequest();
+		FPlatformProcess::Sleep(1); // Make sure the first request started
+		HttpRequestQueuing->CancelRequest();
+	});
+
+	ThreadedHttpRunnable.StartTestHttpThread(true/*bBlockGameThread*/);
+}
+
 #if UE_HTTP_CONNECTION_TIMEOUT_SUPPORT_RETRY
 TEST_CASE_METHOD(FWaitUntilCompleteHttpFixture, "Cancel http request connect before timeout", HTTP_TAG)
 {
diff --git a/Engine/Source/Runtime/Online/HTTP/Private/Apple/AppleHttp.cpp b/Engine/Source/Runtime/Online/HTTP/Private/Apple/AppleHttp.cpp
--- a/Engine/Source/Runtime/Online/HTTP/Private/Apple/AppleHttp.cpp
+++ b/Engine/Source/Runtime/Online/HTTP/Private/Apple/AppleHttp.cpp
@@ -11,6 +11,8 @@
 #include "Misc/App.h"
 #include "Misc/Base64.h"
 
+extern TAutoConsoleVariable<bool> CVarHttpSetGeneralFailureReasonFromCommonCode;
+
 /**
  * Class to hold data from delegate implementation notifications.
  */
@@ -843,7 +845,6 @@ void FAppleHttpRequest::FinishRequest()
 
 	TSharedPtr<FAppleHttpResponse> Response = StaticCastSharedPtr<FAppleHttpResponse>(ResponseCommon);
 	bool bSucceeded = (Response && Response->GetStatusFromDelegate() == EHttpRequestStatus::Succeeded);
-	SetStatus(bSucceeded ? EHttpRequestStatus::Succeeded : EHttpRequestStatus::Failed);
 
 	if (bSucceeded)
 	{
@@ -878,7 +879,16 @@ void FAppleHttpRequest::FinishRequest()
 		{
 			ResponseCommon = nullptr;
 		}
-		OnProcessRequestComplete().ExecuteIfBound(SharedThis(this), ResponseCommon, bSucceeded);
+
+		if (CVarHttpSetGeneralFailureReasonFromCommonCode.GetValueOnAnyThread())
+		{
+			HandleRequestFailed(Response);
+		}
+		else
+		{
+			SetStatus(EHttpRequestStatus::Failed);
+			OnProcessRequestComplete().ExecuteIfBound(SharedThis(this), ResponseCommon, bSucceeded);
+		}
 	}
 }
 
diff --git a/Engine/Source/Runtime/Online/HTTP/Private/Curl/CurlHttp.cpp b/Engine/Source/Runtime/Online/HTTP/Private/Curl/CurlHttp.cpp
--- a/Engine/Source/Runtime/Online/HTTP/Private/Curl/CurlHttp.cpp
+++ b/Engine/Source/Runtime/Online/HTTP/Private/Curl/CurlHttp.cpp
@@ -19,6 +19,8 @@
 #include <openssl/ssl.h>
 #endif
 
+extern TAutoConsoleVariable<bool> CVarHttpSetGeneralFailureReasonFromCommonCode;
+
 #if WITH_SSL
 static int SslCertVerify(int PreverifyOk, X509_STORE_CTX* Context)
 {
@@ -1238,50 +1240,78 @@ void FCurlHttpRequest::FinishRequest()
 			}
 		}
 
-		SetStatus(EHttpRequestStatus::Failed);
-
-		// Mark last request attempt as completed but failed
-		if (bCanceled)
+		if (CVarHttpSetGeneralFailureReasonFromCommonCode.GetValueOnAnyThread())
 		{
-			SetFailureReason(EHttpFailureReason::Cancelled);
-		}
-		else if (bTimedOut)
-		{
-			SetFailureReason(EHttpFailureReason::TimedOut);
-		}
-		else if (bActivityTimedOut)
-		{
-			SetFailureReason(EHttpFailureReason::ConnectionError);
-		}
-		else if (bCurlRequestCompleted)
-		{
-			switch (CurlCompletionResult)
+			if (bCurlRequestCompleted)
 			{
-			case CURLE_COULDNT_CONNECT:
-			case CURLE_OPERATION_TIMEDOUT:
-			case CURLE_COULDNT_RESOLVE_PROXY:
-			case CURLE_COULDNT_RESOLVE_HOST:
-			case CURLE_SSL_CONNECT_ERROR:
+				switch (CurlCompletionResult)
+				{
+				case CURLE_COULDNT_CONNECT:
+				case CURLE_OPERATION_TIMEDOUT:
+				case CURLE_COULDNT_RESOLVE_PROXY:
+				case CURLE_COULDNT_RESOLVE_HOST:
+				case CURLE_SSL_CONNECT_ERROR:
 #if WITH_CURL_XCURL
-			case CURLE_SEND_ERROR:
+				case CURLE_SEND_ERROR:
 #endif
-				// report these as connection errors (safe to retry)
-				SetFailureReason(EHttpFailureReason::ConnectionError);
-				break;
-			default:
-				SetFailureReason(EHttpFailureReason::Other);
+					// report these as connection errors (safe to retry)
+					SetFailureReason(EHttpFailureReason::ConnectionError);
+					break;
+				default:
+					break;
+				}
 			}
+
+			HandleRequestFailed(Response);
 		}
 		else
 		{
-			SetFailureReason(EHttpFailureReason::Other);
-		}
-		// Call delegate with failure
-		OnProcessRequestComplete().ExecuteIfBound(SharedThis(this), Response, false);
+			SetStatus(EHttpRequestStatus::Failed);
 
-		//Delegate needs to know about the errors -- so clear out Response (since connection failed) afterwards...
-		ResponseCommon = nullptr;
-		TotalBytesRead = 0;
+			// Mark last request attempt as completed but failed
+			if (bCanceled)
+			{
+				SetFailureReason(EHttpFailureReason::Cancelled);
+			}
+			else if (bTimedOut)
+			{
+				SetFailureReason(EHttpFailureReason::TimedOut);
+			}
+			else if (bActivityTimedOut)
+			{
+				SetFailureReason(EHttpFailureReason::ConnectionError);
+			}
+			else if (bCurlRequestCompleted)
+			{
+				switch (CurlCompletionResult)
+				{
+				case CURLE_COULDNT_CONNECT:
+				case CURLE_OPERATION_TIMEDOUT:
+				case CURLE_COULDNT_RESOLVE_PROXY:
+				case CURLE_COULDNT_RESOLVE_HOST:
+				case CURLE_SSL_CONNECT_ERROR:
+	#if WITH_CURL_XCURL
+				case CURLE_SEND_ERROR:
+	#endif
+					// report these as connection errors (safe to retry)
+					SetFailureReason(EHttpFailureReason::ConnectionError);
+					break;
+				default:
+					SetFailureReason(EHttpFailureReason::Other);
+				}
+			}
+			else
+			{
+				SetFailureReason(EHttpFailureReason::Other);
+			}
+
+			// Call delegate with failure
+			OnProcessRequestComplete().ExecuteIfBound(SharedThis(this), Response, false);
+
+			//Delegate needs to know about the errors -- so clear out Response (since connection failed) afterwards...
+			ResponseCommon = nullptr;
+			TotalBytesRead = 0;
+		}
 	}
 }
 
diff --git a/Engine/Source/Runtime/Online/HTTP/Private/GenericPlatform/HttpRequestCommon.cpp b/Engine/Source/Runtime/Online/HTTP/Private/GenericPlatform/HttpRequestCommon.cpp
--- a/Engine/Source/Runtime/Online/HTTP/Private/GenericPlatform/HttpRequestCommon.cpp
+++ b/Engine/Source/Runtime/Online/HTTP/Private/GenericPlatform/HttpRequestCommon.cpp
@@ -23,6 +23,13 @@ TAutoConsoleVariable<bool> CVarHttpLogJsonResponseOnly(
 
 }
 
+TAutoConsoleVariable<bool> CVarHttpSetGeneralFailureReasonFromCommonCode(
+	TEXT("http.SetGeneralFailureReasonFromCommonCode"),
+	true,
+	TEXT("Temporary hotfixable cvar: when enabled, set general failure reason from common code instead."),
+	ECVF_SaveForNextBoot
+);
+
 FHttpRequestCommon::FHttpRequestCommon()
 	: RequestStartTimeAbsoluteSeconds(FPlatformTime::Seconds())
 	, ActivityTimeoutAt(0.0)
@@ -181,6 +188,36 @@ void FHttpRequestCommon::HandleRequestSucceed(TSharedPtr<IHttpResponse> InRespon
 	FHttpModule::Get().GetHttpManager().RecordStatTimeToConnect(ConnectTime);
 }
 
+void FHttpRequestCommon::HandleRequestFailed(TSharedPtr<IHttpResponse> InResponse)
+{
+	if (CVarHttpSetGeneralFailureReasonFromCommonCode.GetValueOnAnyThread())
+	{
+		if (FailureReason == EHttpFailureReason::None) // Failure reason was not set by platform, will set it here
+		{
+			if (bCanceled)
+			{
+				SetFailureReason(EHttpFailureReason::Cancelled);
+			}
+			else if (bTimedOut)
+			{
+				SetFailureReason(EHttpFailureReason::TimedOut);
+			}
+			else if (!bUsePlatformActivityTimeout && bActivityTimedOut)
+			{
+				SetFailureReason(EHttpFailureReason::ConnectionError);
+			}
+			else
+			{
+				SetFailureReason(EHttpFailureReason::Other);
+			}
+		}
+
+		SetStatus(EHttpRequestStatus::Failed);
+	}
+
+	OnProcessRequestComplete().ExecuteIfBound(SharedThis(this), InResponse, false);
+}
+
 void FHttpRequestCommon::SetStatus(EHttpRequestStatus::Type InCompletionStatus)
 {
 	CompletionStatus = InCompletionStatus;
diff --git a/Engine/Source/Runtime/Online/HTTP/Public/GenericPlatform/HttpRequestCommon.h b/Engine/Source/Runtime/Online/HTTP/Public/GenericPlatform/HttpRequestCommon.h
--- a/Engine/Source/Runtime/Online/HTTP/Public/GenericPlatform/HttpRequestCommon.h
+++ b/Engine/Source/Runtime/Online/HTTP/Public/GenericPlatform/HttpRequestCommon.h
@@ -75,6 +75,7 @@ protected:
 	HTTP_API void FinishRequestNotInHttpManager();
 
 	HTTP_API void HandleRequestSucceed(TSharedPtr<IHttpResponse> InResponse);
+	HTTP_API void HandleRequestFailed(TSharedPtr<IHttpResponse> InResponse);
 
 	HTTP_API void StartActivityTimeoutTimer();
 	HTTP_API void StartActivityTimeoutTimerBy(double DelayToTrigger);
