On specific platform, cache maximum concurrent http requests when init thread, instead of increase the value at runtime, we can't really increase it by hotfix at runtime when it also requires memory pool re-initialization.

diff --git a/Engine/Source/Programs/WebTests/Private/TestHttp.cpp b/Engine/Source/Programs/WebTests/Private/TestHttp.cpp
--- a/Engine/Source/Programs/WebTests/Private/TestHttp.cpp
+++ b/Engine/Source/Programs/WebTests/Private/TestHttp.cpp
@@ -709,7 +709,7 @@ public:
 TEST_CASE_METHOD(FWaitUntilCompleteHttpFixture, "In streaming downloading http request won't trigger response body receive delegate after canceling", HTTP_TAG)
 {
 	TSharedRef<IHttpRequest> HttpRequest = CreateRequest();
-	HttpRequest->SetURL(UrlStreamDownload(30, 1024*1024));
+	HttpRequest->SetURL(UrlStreamDownload(60, 1024*1024));
 
 	TSharedPtr<FUserStreamingClass> UserInstance = MakeShared<FUserStreamingClass>();
 
@@ -738,7 +738,7 @@ TEST_CASE_METHOD(FWaitUntilCompleteHttpFixture, "In streaming downloading http r
 	DisableWarningsInThisTest(); // Failed writing received data to disk/application
 
 	TSharedRef<IHttpRequest> HttpRequest = CreateRequest();
-	HttpRequest->SetURL(UrlStreamDownload(30, 1024*1024));
+	HttpRequest->SetURL(UrlStreamDownload(60, 1024*1024));
 
 	TSharedPtr<FUserStreamingClass> UserInstance = MakeShared<FUserStreamingClass>();
 
@@ -1780,6 +1780,11 @@ TEST_CASE_METHOD(FWaitThreadedHttpFixture, "Cancel http request with ProcessRequ
 
 	std::atomic<bool> bFirstRequestCompleted = false;
 
+	FHttpManager& HttpManager = HttpModule->GetHttpManager();
+	const FHttpStats HttpStats = HttpManager.GetHttpStats();
+	CHECK(HttpStats.RequestsInQueue == 0);
+	CHECK(HttpStats.MaxRequestsInQueue == 0);
+
 	ThreadedHttpRunnable.OnRunFromThread().BindLambda([this, &bFirstRequestCompleted]() {
 		TSharedRef<IHttpRequest> HttpRequestRunning = CreateRequest();
 		HttpRequestRunning->SetURL(UrlStreamDownload(3/*Chunks*/, HTTP_TEST_TIMEOUT_CHUNK_SIZE, 1/*ChunkLatency*/));
@@ -1812,6 +1817,12 @@ TEST_CASE_METHOD(FWaitThreadedHttpFixture, "Cancel http request with ProcessRequ
 		});
 		HttpRequestQueuing->ProcessRequest();
 		FPlatformProcess::Sleep(1); // Make sure the first request started
+
+		FHttpManager& HttpManager = HttpModule->GetHttpManager();
+		const FHttpStats HttpStats = HttpManager.GetHttpStats();
+		CHECK(HttpStats.RequestsInQueue == 1);
+		CHECK(HttpStats.MaxRequestsInQueue == 1);
+
 		HttpRequestQueuing->CancelRequest();
 	});
 
diff --git a/Engine/Source/Runtime/Online/HTTP/HTTP.Build.cs b/Engine/Source/Runtime/Online/HTTP/HTTP.Build.cs
--- a/Engine/Source/Runtime/Online/HTTP/HTTP.Build.cs
+++ b/Engine/Source/Runtime/Online/HTTP/HTTP.Build.cs
@@ -7,6 +7,8 @@ public class HTTP : ModuleRules
 	// Currently there is a random event loop crash when shutdown HTTP manager on PC
 	protected virtual bool bPlatformEventLoopEnabledByDefault { get { return !Target.Platform.IsInGroup(UnrealPlatformGroup.Windows); } }
 
+	protected virtual bool bPlatformSupportToIncreaseMaxRequestsAtRuntime { get { return true; } }
+
 	protected virtual bool bPlatformSupportsWinHttp
 	{
 		get
@@ -97,6 +99,7 @@ public class HTTP : ModuleRules
 		}
 
 		PrivateDefinitions.Add("UE_HTTP_EVENT_LOOP_ENABLE_CHANCE_BY_DEFAULT=" + (bPlatformEventLoopEnabledByDefault ? "100" : "0"));
+		PrivateDefinitions.Add("UE_HTTP_SUPPORT_TO_INCREASE_MAX_REQUESTS_AT_RUNTIME=" + (bPlatformSupportToIncreaseMaxRequestsAtRuntime ? "1" : "0"));
 		PrivateDefinitions.Add("WITH_CURL_LIBCURL =" + (bPlatformSupportsLibCurl ? "1" : "0"));
 		PrivateDefinitions.Add("WITH_CURL_XCURL=" + (bPlatformSupportsXCurl ? "1" : "0"));
 		PrivateDefinitions.Add("WITH_CURL_MULTIPOLL=" + (bPlatformSupportsCurlMultiPoll ? "1" : "0"));
diff --git a/Engine/Source/Runtime/Online/HTTP/Private/HttpThread.cpp b/Engine/Source/Runtime/Online/HTTP/Private/HttpThread.cpp
--- a/Engine/Source/Runtime/Online/HTTP/Private/HttpThread.cpp
+++ b/Engine/Source/Runtime/Online/HTTP/Private/HttpThread.cpp
@@ -87,7 +87,9 @@ FHttpThreadBase::FHttpThreadBase()
 	, bIsSingleThread(false)
 	, bIsStopped(true)
 	, CurrentThreadPriority(EThreadPriority::TPri_Num)
+	, MaxConcurrentRequests(CVarHttpMaxConcurrentRequests.GetValueOnAnyThread())
 {
+	UE_LOG(LogInit, Log, TEXT("Creating http thread with maximum %d concurrent requests"), MaxConcurrentRequests);
 }
 
 FHttpThreadBase::~FHttpThreadBase()
@@ -272,7 +274,15 @@ void FHttpThreadBase::StartRequestsWaitingInQueue(TArray<FHttpRequestCommon*>& R
 	// Tick new requests separately from existing RunningThreadedRequests so they get a chance 
 	// to send unaffected by possibly large ElapsedTime above
 	int32 RunningThreadedRequestsCounter = RunningThreadedRequests.Num();
-	const int32 MaxConcurrentRequests = CVarHttpMaxConcurrentRequests.GetValueOnAnyThread();
+
+#if !UE_HTTP_SUPPORT_TO_INCREASE_MAX_REQUESTS_AT_RUNTIME
+	// This will enable shrinking but not growing the max concurrent requests at runtime, on platform where http memory pool was pre-allocated when boot
+	if (CVarHttpMaxConcurrentRequests.GetValueOnAnyThread() < MaxConcurrentRequests)
+#endif
+	{
+		MaxConcurrentRequests = CVarHttpMaxConcurrentRequests.GetValueOnAnyThread();
+	}
+
 	if (RunningThreadedRequestsCounter < MaxConcurrentRequests)
 	{
 		while(RunningThreadedRequestsCounter < MaxConcurrentRequests && !RateLimitedThreadedRequests.IsEmpty())
diff --git a/Engine/Source/Runtime/Online/HTTP/Private/HttpThread.h b/Engine/Source/Runtime/Online/HTTP/Private/HttpThread.h
--- a/Engine/Source/Runtime/Online/HTTP/Private/HttpThread.h
+++ b/Engine/Source/Runtime/Online/HTTP/Private/HttpThread.h
@@ -166,6 +166,9 @@ private:
 	/** Current thread priority of the thread. Used to detect when a priority change is requested */
 	EThreadPriority CurrentThreadPriority;
 
+	/** Max concurrent requests the thread can run, the rest of the requests will wait in the queue */
+	int32 MaxConcurrentRequests;
+
 protected:
 PRAGMA_DISABLE_DEPRECATION_WARNINGS
 	/** 
