Extract the user agent header populating code from implementions into FHttpRequestCommon.

#tests Reproduced with new test case in WebTests and passed the test with the changelist.

diff --git a/Engine/Source/Programs/WebTests/Private/TestHttp.cpp b/Engine/Source/Programs/WebTests/Private/TestHttp.cpp
--- a/Engine/Source/Programs/WebTests/Private/TestHttp.cpp
+++ b/Engine/Source/Programs/WebTests/Private/TestHttp.cpp
@@ -2331,6 +2331,41 @@ TEST_CASE_METHOD(FLocalHttpServerFixture, "Local http server can serve large fil
 
 #endif // UE_HTTP_SUPPORT_LOCAL_SERVER
 
+TEST_CASE_METHOD(FWaitUntilCompleteHttpFixture, "Accessing request initial information without issue while request is running", HTTP_TAG)
+{
+	for (int32 i = 0; i < 30; ++i) // Use 2 "for" loops so it doesn't trigger the warning the request waited for too long in the queue
+	{
+		TArray<TSharedRef<IHttpRequest>> Requests;
+		for (int32 j = 0; j < 30; ++j)
+		{
+			TSharedRef<IHttpRequest> HttpRequest = CreateRequest();
+			HttpRequest->SetHeader(TEXT("Custom-HeaderA"), TEXT("a"));
+			HttpRequest->SetHeader(TEXT("Custom-HeaderB"), TEXT("b"));
+			HttpRequest->SetHeader(TEXT("Custom-HeaderC"), TEXT("c"));
+			HttpRequest->SetURL(UrlStreamDownload(3, HTTP_TEST_TIMEOUT_CHUNK_SIZE, 0));
+			HttpRequest->SetDelegateThreadPolicy(EHttpRequestDelegateThreadPolicy::CompleteOnHttpThread);
+			HttpRequest->ProcessRequest();
+			Requests.Add(HttpRequest);
+		}
+
+		bool bRequestsStillRunning = true;
+		while (bRequestsStillRunning)
+		{
+			bRequestsStillRunning = false;
+			for (TSharedRef<IHttpRequest> Request : Requests)
+			{
+				if (!EHttpRequestStatus::IsFinished(Request->GetStatus()))
+				{
+					bRequestsStillRunning = true;
+
+					CHECK(!Request->GetAllHeaders().IsEmpty());
+					CHECK(!Request->GetURL().IsEmpty());
+				}
+			}
+		}
+	}
+}
+
 TEST_CASE_METHOD(FWaitUntilCompleteHttpFixture, "Test platform request requests limits", HTTP_TAG "[LIMIT]")
 {
 	bool bCheckCancel = GENERATE(false, true);
diff --git a/Engine/Source/Runtime/Online/HTTP/Private/Apple/AppleHttp.cpp b/Engine/Source/Runtime/Online/HTTP/Private/Apple/AppleHttp.cpp
--- a/Engine/Source/Runtime/Online/HTTP/Private/Apple/AppleHttp.cpp
+++ b/Engine/Source/Runtime/Online/HTTP/Private/Apple/AppleHttp.cpp
@@ -772,15 +772,6 @@ bool FAppleHttpRequest::SetupRequest()
 		[Request setValue:[NSString stringWithFormat:@"%llu", GetContentLength()] forHTTPHeaderField:@"Content-Length"];
 	}
 
-	const FString CurrentUserAgent = GetHeader("User-Agent");
-	if(CurrentUserAgent.IsEmpty())
-	{
-		FString DefaultUserAgent = FPlatformHttp::GetDefaultUserAgent();
-		UE_LOG(LogHttp, Verbose, TEXT("Setting default User-Agent %s"), *DefaultUserAgent);
-		NSString* UserAgent = DefaultUserAgent.GetNSString();
-		[Request setValue:UserAgent forHTTPHeaderField:@"User-Agent"];
-	}
-
 	PostProcess();
 
 	LastReportedBytesWritten = 0;
diff --git a/Engine/Source/Runtime/Online/HTTP/Private/Curl/CurlHttp.cpp b/Engine/Source/Runtime/Online/HTTP/Private/Curl/CurlHttp.cpp
--- a/Engine/Source/Runtime/Online/HTTP/Private/Curl/CurlHttp.cpp
+++ b/Engine/Source/Runtime/Online/HTTP/Private/Curl/CurlHttp.cpp
@@ -825,11 +825,6 @@ bool FCurlHttpRequest::SetupRequest()
 	UE_LOG(LogHttp, Verbose, TEXT("%p: Custom headers are %s"), this, Headers.Num() ? TEXT("present") : TEXT("NOT present"));
 	UE_LOG(LogHttp, Verbose, TEXT("%p: Payload size=%llu"), this, RequestPayload->GetContentLength());
 
-	if (GetHeader(TEXT("User-Agent")).IsEmpty())
-	{
-		SetHeader(TEXT("User-Agent"), FPlatformHttp::GetDefaultUserAgent());
-	}
-
 	// content-length should be present http://www.w3.org/Protocols/rfc2616/rfc2616-sec4.html#sec4.4
 	if (GetHeader(TEXT("Content-Length")).IsEmpty())
 	{
diff --git a/Engine/Source/Runtime/Online/HTTP/Private/GenericPlatform/HttpRequestCommon.cpp b/Engine/Source/Runtime/Online/HTTP/Private/GenericPlatform/HttpRequestCommon.cpp
--- a/Engine/Source/Runtime/Online/HTTP/Private/GenericPlatform/HttpRequestCommon.cpp
+++ b/Engine/Source/Runtime/Online/HTTP/Private/GenericPlatform/HttpRequestCommon.cpp
@@ -145,6 +145,14 @@ void FHttpRequestCommon::InitResponse()
 	}
 }
 
+void FHttpRequestCommon::PopulateUserAgentHeader()
+{
+	if (GetHeader(TEXT("User-Agent")).IsEmpty())
+	{
+		SetHeader(TEXT("User-Agent"), FPlatformHttp::GetDefaultUserAgent());
+	}
+}
+
 bool FHttpRequestCommon::PreProcess()
 {
 	ClearInCaseOfRetry();
@@ -163,6 +171,8 @@ bool FHttpRequestCommon::PreProcess()
 		return false;
 	}
 
+	PopulateUserAgentHeader();
+
 	if (!SetupRequest())
 	{
 		FinishRequestNotInHttpManager();
diff --git a/Engine/Source/Runtime/Online/HTTP/Public/GenericPlatform/HttpRequestCommon.h b/Engine/Source/Runtime/Online/HTTP/Public/GenericPlatform/HttpRequestCommon.h
--- a/Engine/Source/Runtime/Online/HTTP/Public/GenericPlatform/HttpRequestCommon.h
+++ b/Engine/Source/Runtime/Online/HTTP/Public/GenericPlatform/HttpRequestCommon.h
@@ -62,6 +62,7 @@ protected:
 	 */
 	HTTP_API bool PreProcess();
 	HTTP_API void PostProcess();
+	HTTP_API void PopulateUserAgentHeader();
 	HTTP_API virtual bool SetupRequest() = 0;
 	HTTP_API bool PreCheck() const;
 	HTTP_API virtual void ClearInCaseOfRetry();
