Add support to IHttpRequest to specify a Unix Socket to allow intra host http requests through this mechanism

diff --git a/Engine/Source/Programs/WebTests/Private/TestHttp.cpp b/Engine/Source/Programs/WebTests/Private/TestHttp.cpp
--- a/Engine/Source/Programs/WebTests/Private/TestHttp.cpp
+++ b/Engine/Source/Programs/WebTests/Private/TestHttp.cpp
@@ -156,6 +156,7 @@ public:
 		FParse::Value(FCommandLine::Get(), TEXT("web_server_ip="), WebServerIp);
 		FParse::Bool(FCommandLine::Get(), TEXT("run_heavy_tests="), bRunHeavyTests);
 		FParse::Bool(FCommandLine::Get(), TEXT("retry_enabled="), bRetryEnabled);
+		FParse::Value(FCommandLine::Get(), TEXT("web_server_unix_socket="), WebServerUnixSocket);
 	}
 
 	void DisableWarningsInThisTest()
@@ -181,8 +182,10 @@ public:
 	const FString UrlStreamUpload() { return FString::Format(TEXT("{0}/streaming_upload_put"), { *UrlHttpTests() }); }
 	const FString UrlMockLatency(uint32 Latency) const { return FString::Format(TEXT("{0}/mock_latency/{1}/"), { *UrlHttpTests(), Latency }); }
 	const FString UrlMockStatus(uint32 StatusCode) const { return FString::Format(TEXT("{0}/mock_status/{1}/"), { *UrlHttpTests(), StatusCode }); }
+	const FString UrlUnixSocketHttpTests() const { return "http://localhost/webtests/unixsockettests"; }
 
 	FString WebServerIp;
+	FString WebServerUnixSocket;
 	uint32 WebServerHttpPort;
 	FMockHttpModule* HttpModule = nullptr;
 	bool bRunHeavyTests;
@@ -2423,4 +2426,57 @@ TEST_CASE_METHOD(FWaitUntilCompleteHttpFixture, "Test platform request requests
 			}
 		}
 	}
-}
\ No newline at end of file
+}
+
+#if UE_HTTP_SUPPORT_UNIX_SOCKET
+
+TEST_CASE_METHOD(FWaitUntilCompleteHttpFixture, "Http Methods over Unix Domain Socket", HTTP_TAG)
+{
+	if (WebServerUnixSocket.Len() == 0)
+	{
+		return;
+	}
+
+	TSharedRef<IHttpRequest> HttpRequest = CreateRequest();
+	CHECK(HttpRequest->GetVerb() == TEXT("GET"));
+
+	const int Number = FPlatformTime::Cycles();
+
+	HttpRequest->SetURL(FString::Format(TEXT("{0}/{1}"), { *UrlUnixSocketHttpTests(), Number }));
+	HttpRequest->SetOption(HttpRequestOptions::UnixSocketPath, WebServerUnixSocket);
+
+	SECTION("Default GET")
+	{
+	}
+	SECTION("GET")
+	{
+		HttpRequest->SetVerb(TEXT("GET"));
+	}
+	SECTION("POST")
+	{
+		HttpRequest->SetVerb(TEXT("POST"));
+	}
+	SECTION("PUT")
+	{
+		HttpRequest->SetVerb(TEXT("PUT"));
+	}
+	SECTION("DELETE")
+	{
+		HttpRequest->SetVerb(TEXT("DELETE"));
+	}
+
+	HttpRequest->OnProcessRequestComplete().BindLambda([Number](FHttpRequestPtr HttpRequest, FHttpResponsePtr HttpResponse, bool bSucceeded) {
+		CHECK(bSucceeded);
+		REQUIRE(HttpResponse != nullptr);
+		CHECK(HttpResponse->GetResponseCode() == 200);
+
+		FString ResponseContent = HttpResponse->GetContentAsString();
+
+		int NumberReturned = FCString::Atoi(*ResponseContent);
+		CHECK(Number == NumberReturned);
+
+		});
+	HttpRequest->ProcessRequest();
+}
+
+#endif //UE_HTTP_SUPPORT_UNIX_SOCKET
\ No newline at end of file
diff --git a/Engine/Source/Programs/WebTestsServer/runserver_unixsocket.bat b/Engine/Source/Programs/WebTestsServer/runserver_unixsocket.bat
new file mode 100755
--- /dev/null
+++ b/Engine/Source/Programs/WebTestsServer/runserver_unixsocket.bat
@@ -0,0 +1,2 @@
+dotnet build .\unixsockets\unixsocketsserver.csproj -c "Release"
+dotnet run --project ".\unixsockets\unixsocketsserver.csproj" -c "Release"
\ No newline at end of file
diff --git a/Engine/Source/Programs/WebTestsServer/unixsockets/Program.cs b/Engine/Source/Programs/WebTestsServer/unixsockets/Program.cs
new file mode 100644
--- /dev/null
+++ b/Engine/Source/Programs/WebTestsServer/unixsockets/Program.cs
@@ -0,0 +1,48 @@
+﻿// Copyright Epic Games, Inc. All Rights Reserved.
+using Microsoft.AspNetCore.Builder;
+using Microsoft.AspNetCore.Hosting;
+
+const string SocketFileName = "webtests.sock";
+
+string CurrentDirectory = Environment.CurrentDirectory;
+string SocketsDirectory = Path.Combine(CurrentDirectory, "sockets");
+
+if (Directory.Exists(SocketsDirectory))
+{
+    //Clean any sockets remaining from a previous run
+    foreach(var SocketFile in Directory.EnumerateFiles(SocketsDirectory, "*", SearchOption.AllDirectories))
+    {
+        File.Delete(SocketFile);
+    }
+}
+else
+{
+    Directory.CreateDirectory(SocketsDirectory);
+}
+
+string SocketPath = Path.Combine(SocketsDirectory, SocketFileName);
+Console.WriteLine($"Run Web Tests with the following arguments to test Unix Sockets");
+Console.WriteLine($"=> \"Http Methods over Unix Domain Socket\" --extra-args -very_verbose=true -web_server_unix_socket=\"{SocketPath}\"");
+
+var AppBuilder = WebApplication.CreateBuilder();
+AppBuilder.WebHost.UseKestrel(Options =>
+{
+    Options.ListenUnixSocket(SocketPath);
+});
+
+var WebServerApp = AppBuilder.Build();
+
+string[] Methods = { "GET", "POST", "PUT", "DELETE" };
+
+WebServerApp.MapMethods("webtests/unixsockettests/{param}", Methods, async Context =>
+{
+    string Param = (string)Context.Request.RouteValues["param"];
+    Context.Response.StatusCode = 200;
+    await using (StreamWriter Writer = new StreamWriter(Context.Response.Body, leaveOpen:true))
+    {
+        await Writer.WriteAsync(Param);
+    }
+    await Task.CompletedTask;
+});
+
+WebServerApp.Run();
\ No newline at end of file
diff --git a/Engine/Source/Programs/WebTestsServer/unixsockets/unixsocketsserver.csproj b/Engine/Source/Programs/WebTestsServer/unixsockets/unixsocketsserver.csproj
new file mode 100644
--- /dev/null
+++ b/Engine/Source/Programs/WebTestsServer/unixsockets/unixsocketsserver.csproj
@@ -0,0 +1,14 @@
+﻿<Project Sdk="Microsoft.NET.Sdk">
+
+  <PropertyGroup>
+    <OutputType>Exe</OutputType>
+    <TargetFramework>net8.0</TargetFramework>
+    <ImplicitUsings>enable</ImplicitUsings>
+    <Nullable>disable</Nullable>
+  </PropertyGroup>
+
+  <ItemGroup>
+    <FrameworkReference Include="Microsoft.AspNetCore.App" />
+  </ItemGroup>
+
+</Project>
diff --git a/Engine/Source/Runtime/Online/BuildPatchServices/Private/Tests/Mock/HttpRequest.mock.h b/Engine/Source/Runtime/Online/BuildPatchServices/Private/Tests/Mock/HttpRequest.mock.h
--- a/Engine/Source/Runtime/Online/BuildPatchServices/Private/Tests/Mock/HttpRequest.mock.h
+++ b/Engine/Source/Runtime/Online/BuildPatchServices/Private/Tests/Mock/HttpRequest.mock.h
@@ -71,6 +71,12 @@ namespace BuildPatchServices
 			return FString();
 		}
 
+		virtual FString GetOption(const FName Option) const override
+		{
+			MOCK_FUNC_NOT_IMPLEMENTED("FMockHttpRequest::GetOption");
+			return FString();
+		}
+
 		virtual void SetVerb(const FString& Verb) override
 		{
 			RxSetVerb.Emplace(Verb);
@@ -81,6 +87,11 @@ namespace BuildPatchServices
 			RxSetURL.Emplace(URL);
 		}
 
+		virtual void SetOption(const FName Option, const FString& OptionValue) override
+		{
+			MOCK_FUNC_NOT_IMPLEMENTED("FMockHttpRequest::SetOption");
+		}
+
 		virtual void SetContent(const TArray<uint8>& ContentPayload) override
 		{
 			MOCK_FUNC_NOT_IMPLEMENTED("FMockHttpRequest::SetContent");
diff --git a/Engine/Source/Runtime/Online/HTTP/HTTP.Build.cs b/Engine/Source/Runtime/Online/HTTP/HTTP.Build.cs
--- a/Engine/Source/Runtime/Online/HTTP/HTTP.Build.cs
+++ b/Engine/Source/Runtime/Online/HTTP/HTTP.Build.cs
@@ -58,6 +58,15 @@ public class HTTP : ModuleRules
 		}
 	}
 
+	protected virtual bool bPlatformSupportsUnixSockets
+	{
+		get
+		{
+			return (Target.Platform.IsInGroup(UnrealPlatformGroup.Windows) && !Target.WindowsPlatform.bUseXCurl) ||
+				Target.IsInPlatformGroup(UnrealPlatformGroup.Unix);
+		}
+	}
+
 	public HTTP(ReadOnlyTargetRules Target) : base(Target)
 	{
 		PublicDefinitions.Add("HTTP_PACKAGE=1");
@@ -149,5 +158,6 @@ public class HTTP : ModuleRules
 		PublicDefinitions.Add("UE_HTTP_CONNECTION_TIMEOUT_SUPPORT_RETRY=" + (bPlatformConnectionTimeoutSupportRetry ? "1" : "0"));
 		PublicDefinitions.Add("UE_HTTP_ACTIVITY_TIMER_START_AFTER_RECEIVED_DATA=" + ((bPlatformSupportsXCurl || Target.IsInPlatformGroup(UnrealPlatformGroup.Apple)) ? "1" : "0"));
 		PublicDefinitions.Add("UE_HTTP_SUPPORT_LOCAL_SERVER=" + (bPlatformSupportsLocalHttpServer ? "1" : "0"));
+		PublicDefinitions.Add("UE_HTTP_SUPPORT_UNIX_SOCKET=" + (bPlatformSupportsUnixSockets ? "1" : "0"));
 	}
 }
diff --git a/Engine/Source/Runtime/Online/HTTP/Private/Curl/CurlHttp.cpp b/Engine/Source/Runtime/Online/HTTP/Private/Curl/CurlHttp.cpp
--- a/Engine/Source/Runtime/Online/HTTP/Private/Curl/CurlHttp.cpp
+++ b/Engine/Source/Runtime/Online/HTTP/Private/Curl/CurlHttp.cpp
@@ -312,6 +312,18 @@ void FCurlHttpRequest::SetURL(const FString& InURL)
 	URL = InURL;
 }
 
+void FCurlHttpRequest::SetOption(const FName Option, const FString& OptionValue)
+{
+	if (CompletionStatus == EHttpRequestStatus::Processing)
+	{
+		UE_LOG(LogHttp, Warning, TEXT("FCurlHttpRequest::SetOption() - attempted to set option on a request that is inflight"));
+		return;
+	}
+
+	check(EasyHandle);
+	FHttpRequestCommon::SetOption(Option, OptionValue);
+}
+
 void FCurlHttpRequest::SetContent(const TArray<uint8>& ContentPayload)
 {
 	SetContent(CopyTemp(ContentPayload));
@@ -869,6 +881,14 @@ bool FCurlHttpRequest::SetupRequestHttpThread()
 
 		curl_easy_setopt(EasyHandle, CURLOPT_URL, TCHAR_TO_ANSI(*URL));
 
+#if UE_HTTP_SUPPORT_UNIX_SOCKET
+		FString UnixSocketPath = GetOption(HttpRequestOptions::UnixSocketPath);
+		if (UnixSocketPath.Len() > 0)
+		{
+			curl_easy_setopt(EasyHandle, CURLOPT_UNIX_SOCKET_PATH, TCHAR_TO_ANSI(*UnixSocketPath));
+		}
+#endif //UE_HTTP_SUPPORT_UNIX_SOCKET
+
 		if (!FCurlHttpManager::CurlRequestOptions.LocalHostAddr.IsEmpty())
 		{
 			// Set the local address to use for making these requests
@@ -1447,4 +1467,4 @@ FString FCurlHttpResponse::GetContentAsString() const
 
 PRAGMA_ENABLE_DEPRECATION_WARNINGS
 
-#endif //WITH_CURL
+#endif //WITH_CURL
\ No newline at end of file
diff --git a/Engine/Source/Runtime/Online/HTTP/Private/Curl/CurlHttp.h b/Engine/Source/Runtime/Online/HTTP/Private/Curl/CurlHttp.h
--- a/Engine/Source/Runtime/Online/HTTP/Private/Curl/CurlHttp.h
+++ b/Engine/Source/Runtime/Online/HTTP/Private/Curl/CurlHttp.h
@@ -157,6 +157,7 @@ public:
 	virtual FString GetVerb() const override;
 	virtual void SetVerb(const FString& InVerb) override;
 	virtual void SetURL(const FString& InURL) override;
+	virtual void SetOption(const FName Option, const FString& OptionValue) override;
 	virtual void SetContent(const TArray<uint8>& ContentPayload) override;
 	virtual void SetContent(TArray<uint8>&& ContentPayload) override;
 	virtual void SetContentAsString(const FString& ContentString) override;
diff --git a/Engine/Source/Runtime/Online/HTTP/Private/GenericPlatform/GenericPlatformHttp.cpp b/Engine/Source/Runtime/Online/HTTP/Private/GenericPlatform/GenericPlatformHttp.cpp
--- a/Engine/Source/Runtime/Online/HTTP/Private/GenericPlatform/GenericPlatformHttp.cpp
+++ b/Engine/Source/Runtime/Online/HTTP/Private/GenericPlatform/GenericPlatformHttp.cpp
@@ -107,6 +107,8 @@ public:
 	virtual FString GetVerb() const override { return TEXT(""); }
 	virtual void SetVerb(const FString& Verb) override {}
 	virtual void SetURL(const FString& URL) override {}
+	virtual FString GetOption(const FName Option) const override { return TEXT(""); }
+	virtual void SetOption(const FName Option, const FString& OptionValue) override {}
 	virtual void SetContent(const TArray<uint8>& ContentPayload) override {}
 	virtual void SetContent(TArray<uint8>&& ContentPayload) override {}
 	virtual void SetContentAsString(const FString& ContentString) override {}
diff --git a/Engine/Source/Runtime/Online/HTTP/Private/GenericPlatform/HttpRequestCommon.cpp b/Engine/Source/Runtime/Online/HTTP/Private/GenericPlatform/HttpRequestCommon.cpp
--- a/Engine/Source/Runtime/Online/HTTP/Private/GenericPlatform/HttpRequestCommon.cpp
+++ b/Engine/Source/Runtime/Online/HTTP/Private/GenericPlatform/HttpRequestCommon.cpp
@@ -242,6 +242,21 @@ EHttpRequestDelegateThreadPolicy FHttpRequestCommon::GetDelegateThreadPolicy() c
 	return DelegateThreadPolicy; 
 }
 
+FString FHttpRequestCommon::GetOption(const FName Option) const
+{
+	const FString* OptionValue = Options.Find(Option);
+	if (OptionValue)
+	{
+		return *OptionValue;
+	}
+	return TEXT("");
+}
+
+void FHttpRequestCommon::SetOption(const FName Option, const FString& OptionValue)
+{
+	Options.Add(Option, OptionValue);
+}
+
 void FHttpRequestCommon::HandleRequestSucceed(TSharedPtr<IHttpResponse> InResponse)
 {
 	SetStatus(EHttpRequestStatus::Succeeded);
diff --git a/Engine/Source/Runtime/Online/HTTP/Private/HttpRequestAdapter.cpp b/Engine/Source/Runtime/Online/HTTP/Private/HttpRequestAdapter.cpp
--- a/Engine/Source/Runtime/Online/HTTP/Private/HttpRequestAdapter.cpp
+++ b/Engine/Source/Runtime/Online/HTTP/Private/HttpRequestAdapter.cpp
@@ -49,6 +49,11 @@ FString FHttpRequestAdapterBase::GetVerb() const
 	return HttpRequest->GetVerb(); 
 }
 
+FString FHttpRequestAdapterBase::GetOption(const FName Option) const
+{
+	return HttpRequest->GetOption(Option);
+}
+
 void FHttpRequestAdapterBase::SetVerb(const FString& Verb) 
 { 
 	HttpRequest->SetVerb(Verb); 
@@ -59,6 +64,11 @@ void FHttpRequestAdapterBase::SetURL(const FString& URL)
 	HttpRequest->SetURL(URL); 
 }
 
+void FHttpRequestAdapterBase::SetOption(const FName Option, const FString& OptionValue)
+{
+	HttpRequest->SetOption(Option, OptionValue);
+}
+
 void FHttpRequestAdapterBase::SetContent(const TArray<uint8>& ContentPayload) 
 { 
 	HttpRequest->SetContent(ContentPayload); 
diff --git a/Engine/Source/Runtime/Online/HTTP/Private/NullHttp.cpp b/Engine/Source/Runtime/Online/HTTP/Private/NullHttp.cpp
--- a/Engine/Source/Runtime/Online/HTTP/Private/NullHttp.cpp
+++ b/Engine/Source/Runtime/Online/HTTP/Private/NullHttp.cpp
@@ -67,6 +67,21 @@ void FNullHttpRequest::SetURL(const FString& InURL)
 	Url = InURL;
 }
 
+FString FNullHttpRequest::GetOption(const FName Option) const
+{
+	const FString* OptionValue = Options.Find(Option);
+	if (OptionValue)
+	{
+		return *OptionValue;
+	}
+	return TEXT("");
+}
+
+void FNullHttpRequest::SetOption(const FName Option, const FString& OptionValue)
+{
+	Options.Add(Option, OptionValue);
+}
+
 void FNullHttpRequest::SetContent(const TArray<uint8>& ContentPayload)
 {
 	Payload = ContentPayload;
diff --git a/Engine/Source/Runtime/Online/HTTP/Private/NullHttp.h b/Engine/Source/Runtime/Online/HTTP/Private/NullHttp.h
--- a/Engine/Source/Runtime/Online/HTTP/Private/NullHttp.h
+++ b/Engine/Source/Runtime/Online/HTTP/Private/NullHttp.h
@@ -25,6 +25,8 @@ public:
 	virtual FString GetVerb() const override;
 	virtual void SetVerb(const FString& InVerb) override;
 	virtual void SetURL(const FString& InURL) override;
+	virtual FString GetOption(const FName Option) const override;
+	virtual void SetOption(const FName Option, const FString& OptionValue) override;
 	virtual void SetContent(const TArray<uint8>& ContentPayload) override;
 	virtual void SetContent(TArray<uint8>&& ContentPayload) override;
 	virtual void SetContentAsString(const FString& ContentString) override;
@@ -67,6 +69,7 @@ private:
 	EHttpRequestStatus::Type CompletionStatus;
 	EHttpFailureReason FailureReason;
 	TMap<FString, FString> Headers;
+	TMap<const FName, FString> Options;
 	float ElapsedTime;
 	TOptional<float> TimeoutSecs;
 };
diff --git a/Engine/Source/Runtime/Online/HTTP/Public/GenericPlatform/HttpRequestCommon.h b/Engine/Source/Runtime/Online/HTTP/Public/GenericPlatform/HttpRequestCommon.h
--- a/Engine/Source/Runtime/Online/HTTP/Public/GenericPlatform/HttpRequestCommon.h
+++ b/Engine/Source/Runtime/Online/HTTP/Public/GenericPlatform/HttpRequestCommon.h
@@ -25,6 +25,8 @@ public:
 	HTTP_API virtual EHttpFailureReason GetFailureReason() const override;
 	HTTP_API virtual void SetDelegateThreadPolicy(EHttpRequestDelegateThreadPolicy InDelegateThreadPolicy) override;
 	HTTP_API virtual EHttpRequestDelegateThreadPolicy GetDelegateThreadPolicy() const override;
+	HTTP_API virtual FString GetOption(const FName Option) const override;
+	HTTP_API virtual void SetOption(const FName Option, const FString & OptionValue) override;
 
 	HTTP_API virtual void SetTimeout(float InTimeoutSecs) override;
 	HTTP_API virtual void ClearTimeout() override;
@@ -180,6 +182,9 @@ PRAGMA_ENABLE_DEPRECATION_WARNINGS
 	/** Payload to use with the request. Typically for POST, PUT, or PATCH */
 	TUniquePtr<FRequestPayload> RequestPayload;
 
+	/** Options for this request. */
+	TMap<const FName, FString> Options;
+
 	/** Total elapsed time in seconds since the start of the request */
 	float ElapsedTime = 0.0f;
 
diff --git a/Engine/Source/Runtime/Online/HTTP/Public/HttpRequestAdapter.h b/Engine/Source/Runtime/Online/HTTP/Public/HttpRequestAdapter.h
--- a/Engine/Source/Runtime/Online/HTTP/Public/HttpRequestAdapter.h
+++ b/Engine/Source/Runtime/Online/HTTP/Public/HttpRequestAdapter.h
@@ -25,8 +25,10 @@ public:
 	HTTP_API virtual uint64 GetContentLength() const override;
 	HTTP_API virtual const TArray<uint8>& GetContent() const override;
 	HTTP_API virtual FString GetVerb() const override;
+	HTTP_API virtual FString GetOption(const FName Option) const override;
 	HTTP_API virtual void SetVerb(const FString& Verb) override;
 	HTTP_API virtual void SetURL(const FString& URL) override;
+	HTTP_API virtual void SetOption(const FName Option, const FString& OptionValue) override;
 	HTTP_API virtual void SetContent(const TArray<uint8>& ContentPayload) override;
 	HTTP_API virtual void SetContent(TArray<uint8>&& ContentPayload) override;
 	HTTP_API virtual void SetContentAsString(const FString& ContentString) override;
diff --git a/Engine/Source/Runtime/Online/HTTP/Public/Interfaces/IHttpRequest.h b/Engine/Source/Runtime/Online/HTTP/Public/Interfaces/IHttpRequest.h
--- a/Engine/Source/Runtime/Online/HTTP/Public/Interfaces/IHttpRequest.h
+++ b/Engine/Source/Runtime/Online/HTTP/Public/Interfaces/IHttpRequest.h
@@ -125,6 +125,16 @@ private:
 	FHttpRequestStreamDelegate StreamDelegate;
 };
 
+/**
+ * Options that can be specified on a Http Request
+ */
+namespace HttpRequestOptions
+{
+#if UE_HTTP_SUPPORT_UNIX_SOCKET
+	static const FName UnixSocketPath("UnixSocketPath");
+#endif //UE_HTTP_SUPPORT_UNIX_SOCKET
+}
+
 /**
  * Interface for Http requests (created using FHttpFactory)
  */
@@ -159,6 +169,22 @@ public:
 	 */
 	virtual void SetURL(const FString& URL) = 0;
 
+	/**
+	 * Get the current value for the given option
+	 *
+	 * @return the current value set for this option or an empty string if no value has been specified
+	 */
+	virtual FString GetOption(const FName Option) const = 0;
+
+	/**
+	 * Sets the given option for this Request
+	 * Must be set before calling ProcessRequest.
+	 *
+	 * @param Option - The option to set, see 'HttpRequestOptions' for supported options
+	 * @param OptionValue - The value of the option to set
+	 */
+	virtual void SetOption(const FName Option, const FString& OptionValue) = 0;
+
 	/**
 	 * Sets the content of the request (optional data).
 	 * Usually only set for POST requests.
