Add custom verb supports.

diff --git a/Engine/Source/Programs/WebTests/Private/TestHttp.cpp b/Engine/Source/Programs/WebTests/Private/TestHttp.cpp
--- a/Engine/Source/Programs/WebTests/Private/TestHttp.cpp
+++ b/Engine/Source/Programs/WebTests/Private/TestHttp.cpp
@@ -352,6 +352,19 @@ public:
 	uint32 ExpectingExtraCallbacks = 0;
 };
 
+TEST_CASE_METHOD(FWaitUntilCompleteHttpFixture, "Http method use Get by default", HTTP_TAG)
+{
+	TSharedRef<IHttpRequest> HttpRequest = CreateRequest();
+	CHECK(HttpRequest->GetVerb() == TEXT("GET"));
+	HttpRequest->SetURL(UrlToTestMethods());
+	HttpRequest->OnProcessRequestComplete().BindLambda([](FHttpRequestPtr HttpRequest, FHttpResponsePtr HttpResponse, bool bSucceeded) {
+		CHECK(bSucceeded);
+		REQUIRE(HttpResponse != nullptr);
+		CHECK(HttpResponse->GetResponseCode() == 200);
+	});
+	HttpRequest->ProcessRequest();
+}
+
 TEST_CASE_METHOD(FWaitUntilCompleteHttpFixture, "Http Methods", HTTP_TAG)
 {
 	TSharedRef<IHttpRequest> HttpRequest = CreateRequest();
@@ -359,24 +372,15 @@ TEST_CASE_METHOD(FWaitUntilCompleteHttpFixture, "Http Methods", HTTP_TAG)
 
 	HttpRequest->SetURL(UrlToTestMethods());
 
-	SECTION("Default GET")
-	{
-	}
-	SECTION("GET")
-	{
-		HttpRequest->SetVerb(TEXT("GET"));
-	}
-	SECTION("POST")
-	{
-		HttpRequest->SetVerb(TEXT("POST"));
-	}
-	SECTION("PUT")
-	{
-		HttpRequest->SetVerb(TEXT("PUT"));
-	}
-	SECTION("DELETE")
+#if UE_HTTP_SUPPORT_VERB_CONNECT
+	const ANSICHAR* Method = GENERATE("GET", "POST", "DELETE", "PUT", "PATCH", "HEAD", "OPTIONS", "TRACE", "CONNECT");
+#else // Can't use macro UE_HTTP_SUPPORT_VERB_CONNECT inside GENERATE macro, so define it here again
+	const ANSICHAR* Method = GENERATE("GET", "POST", "DELETE", "PUT", "PATCH", "HEAD", "OPTIONS", "TRACE");
+#endif
+
+	DYNAMIC_SECTION("Test methods with " << Method << " request")
 	{
-		HttpRequest->SetVerb(TEXT("DELETE"));
+		HttpRequest->SetVerb(FString(Method));
 	}
 
 	HttpRequest->OnProcessRequestComplete().BindLambda([](FHttpRequestPtr HttpRequest, FHttpResponsePtr HttpResponse, bool bSucceeded) {
diff --git a/Engine/Source/Programs/WebTestsServer/httptests/views.py b/Engine/Source/Programs/WebTestsServer/httptests/views.py
--- a/Engine/Source/Programs/WebTestsServer/httptests/views.py
+++ b/Engine/Source/Programs/WebTestsServer/httptests/views.py
@@ -13,7 +13,7 @@ import logging
 
 logger = logging.getLogger('django')
 
-@api_view(['GET', 'POST', 'DELETE', 'PUT'])
+@api_view(['GET', 'POST', 'DELETE', 'PUT', 'PATCH', 'HEAD', 'OPTIONS', 'CONNECT', 'TRACE'])
 def methods(request):
     return JsonResponse({})
 
diff --git a/Engine/Source/Runtime/Online/HTTP/HTTP.Build.cs b/Engine/Source/Runtime/Online/HTTP/HTTP.Build.cs
--- a/Engine/Source/Runtime/Online/HTTP/HTTP.Build.cs
+++ b/Engine/Source/Runtime/Online/HTTP/HTTP.Build.cs
@@ -27,6 +27,7 @@ public class HTTP : ModuleRules
 	}
 	protected virtual bool bPlatformSupportsXCurl { get { return Target.Platform.IsInGroup(UnrealPlatformGroup.Windows) && Target.WindowsPlatform.bUseXCurl; } }
 	protected virtual bool bPlatformSupportsCurlMultiSocket { get { return !bPlatformSupportsXCurl; } }
+	protected virtual bool bPlatformSupportsVerbConnect { get { return !bPlatformSupportsXCurl; } }
 
 	protected virtual bool bPlatformSupportsCurlMultiPoll { get { return true; } }
 
@@ -158,5 +159,6 @@ public class HTTP : ModuleRules
 		PublicDefinitions.Add("UE_HTTP_ACTIVITY_TIMER_START_AFTER_RECEIVED_DATA=" + ((bPlatformSupportsXCurl || Target.IsInPlatformGroup(UnrealPlatformGroup.Apple)) ? "1" : "0"));
 		PublicDefinitions.Add("UE_HTTP_SUPPORT_LOCAL_SERVER=" + (bPlatformSupportsLocalHttpServer ? "1" : "0"));
 		PublicDefinitions.Add("UE_HTTP_SUPPORT_UNIX_SOCKET=" + (bPlatformSupportsUnixSockets ? "1" : "0"));
+		PublicDefinitions.Add("UE_HTTP_SUPPORT_VERB_CONNECT=" + (bPlatformSupportsVerbConnect ? "1" : "0"));
 	}
 }
diff --git a/Engine/Source/Runtime/Online/HTTP/Private/Curl/CurlHttp.cpp b/Engine/Source/Runtime/Online/HTTP/Private/Curl/CurlHttp.cpp
--- a/Engine/Source/Runtime/Online/HTTP/Private/Curl/CurlHttp.cpp
+++ b/Engine/Source/Runtime/Online/HTTP/Private/Curl/CurlHttp.cpp
@@ -949,8 +949,7 @@ bool FCurlHttpRequest::SetupRequestHttpThread()
 		}
 		else
 		{
-			UE_LOG(LogHttp, Fatal, TEXT("Unsupported verb '%s', can be perhaps added with CURLOPT_CUSTOMREQUEST"), *Verb);
-			UE_DEBUG_BREAK();
+			curl_easy_setopt(EasyHandle, CURLOPT_CUSTOMREQUEST, TCHAR_TO_ANSI(*Verb));
 		}
 
 		if (bUseReadFunction)
diff --git a/Engine/Source/Runtime/Online/HTTP/Private/GenericPlatform/HttpRequestCommon.cpp b/Engine/Source/Runtime/Online/HTTP/Private/GenericPlatform/HttpRequestCommon.cpp
--- a/Engine/Source/Runtime/Online/HTTP/Private/GenericPlatform/HttpRequestCommon.cpp
+++ b/Engine/Source/Runtime/Online/HTTP/Private/GenericPlatform/HttpRequestCommon.cpp
@@ -63,6 +63,10 @@ EHttpFailureReason FHttpRequestCommon::GetFailureReason() const
 
 bool FHttpRequestCommon::PreCheck() const
 {
+#if !UE_HTTP_SUPPORT_VERB_CONNECT
+	checkf(!GetVerb().Equals(FString("CONNECT"), ESearchCase::IgnoreCase), TEXT("CONNECT verb is not supported on this platform."));
+#endif
+
 	// Disabled http request processing
 	if (!FHttpModule::Get().IsHttpEnabled())
 	{
