Fix the WebTests failures caused by invalid alt domain setup.

diff --git a/Engine/Source/Programs/WebTests/Private/TestHttp.cpp b/Engine/Source/Programs/WebTests/Private/TestHttp.cpp
--- a/Engine/Source/Programs/WebTests/Private/TestHttp.cpp
+++ b/Engine/Source/Programs/WebTests/Private/TestHttp.cpp
@@ -175,8 +175,9 @@ public:
 		return bRetryEnabled ? HttpRetryManager->CreateRequest() : HttpModule->CreateRequest();
 	}
 
-	const FString UrlWithInvalidPortToTestConnectTimeout() const { return TEXT("http://10.255.255.1:8765"); } // non-routable IP address with a random port
-	const FString UrlBase() const { return FString::Format(TEXT("http://{0}:{1}"), { *WebServerIp, WebServerHttpPort }); }
+	void SetInvalidDomainToTestConnectTimeout() { WebServerIp = TEXT("10.255.255.1"); WebServerHttpPort = 8765; } // non-routable IP address with a random port
+	const FString UrlDomain() const { return FString::Format(TEXT("{0}:{1}"), { *WebServerIp, WebServerHttpPort }); }
+	const FString UrlBase() const { return FString::Format(TEXT("http://{0}"), { *UrlDomain() }); }
 	const FString UrlHttpTests() const { return FString::Format(TEXT("{0}/webtests/httptests"), { *UrlBase() }); }
 	const FString UrlToTestMethods() const { return FString::Format(TEXT("{0}/methods"), { *UrlHttpTests() }); }
 	const FString UrlStreamDownload(uint32 Chunks, uint32 ChunkSize, uint32 ChunkLatency=0) { return FString::Format(TEXT("{0}/streaming_download/{1}/{2}/{3}/"), { *UrlHttpTests(), Chunks, ChunkSize, ChunkLatency }); }
@@ -224,6 +225,8 @@ TEST_CASE_METHOD(FHttpModuleTestFixture, "Shutdown http module without issue whe
 
 TEST_CASE_METHOD(FHttpModuleTestFixture, "Duplicated scheme protocol will trigger check", HTTP_TAG)
 {
+	DisableWarningsInThisTest();
+
 	TSharedRef<IHttpRequest> HttpRequest = CreateRequest();
 	CHECK(HttpRequest->GetVerb() == TEXT("GET"));
 
@@ -536,7 +539,8 @@ TEST_CASE_METHOD(FWaitUntilCompleteHttpFixture, "Http request connect timeout",
 
 	TSharedRef<IHttpRequest> HttpRequest = CreateRequest();
 
-	HttpRequest->SetURL(UrlWithInvalidPortToTestConnectTimeout());
+	SetInvalidDomainToTestConnectTimeout();
+	HttpRequest->SetURL(UrlToTestMethods());
 	HttpRequest->SetVerb(TEXT("GET"));
 
 	const double StartTime = FPlatformTime::Seconds();
@@ -1341,13 +1345,14 @@ TEST_CASE_METHOD(FWaitUntilQuitFromTestFixture, "Make sure connection time out c
 	float ConnectionTimeoutDuration = 2.0f;
 	HttpModule->HttpConnectionTimeout = ConnectionTimeoutDuration;
 
-	HttpRequest->SetURL(UrlWithInvalidPortToTestConnectTimeout());
+	SetInvalidDomainToTestConnectTimeout();
+	HttpRequest->SetURL(UrlToTestMethods());
 
 	const double StartTime = FPlatformTime::Seconds();
 
 	HttpRequest->OnProcessRequestComplete().BindLambda([this, StartTime, ConnectionTimeoutDuration](FHttpRequestPtr HttpRequest, FHttpResponsePtr HttpResponse, bool bSucceeded) {
 		TSharedRef<IHttpRequest> HttpRequest2 = CreateRequest();
-		HttpRequest2->SetURL(UrlWithInvalidPortToTestConnectTimeout());
+		HttpRequest2->SetURL(UrlToTestMethods());
 		HttpRequest2->OnProcessRequestComplete().BindLambda([this, StartTime, ConnectionTimeoutDuration](FHttpRequestPtr HttpRequest, FHttpResponsePtr HttpResponse, bool bSucceeded) {
 			bQuitRequested = true;
 			const double DurationInSeconds  = FPlatformTime::Seconds() - StartTime;
@@ -1852,7 +1857,8 @@ TEST_CASE_METHOD(FWaitUntilCompleteHttpFixture, "Cancel http request connect bef
 	DisableWarningsInThisTest();
 
 	TSharedRef<IHttpRequest> HttpRequest = CreateRequest();
-	HttpRequest->SetURL(UrlWithInvalidPortToTestConnectTimeout());
+	SetInvalidDomainToTestConnectTimeout();
+	HttpRequest->SetURL(UrlToTestMethods());
 	HttpRequest->SetVerb(TEXT("GET"));
 	HttpRequest->SetTimeout(7);
 	const double StartTime = FPlatformTime::Seconds();
@@ -2061,7 +2067,13 @@ TEST_CASE_METHOD(FWaitUntilCompleteHttpFixture, "Retry immediately without lock
 	HttpModule->HttpConnectionTimeout = 1.0f;
 
 	TArray<FString> AltDomains;
-	AltDomains.Add(UrlToTestMethods());
+	FString CorrectDomain = UrlDomain();
+
+	SetInvalidDomainToTestConnectTimeout();
+
+	FString InvalidDomain = UrlDomain();
+	AltDomains.Add(InvalidDomain);
+	AltDomains.Add(CorrectDomain);
 	FHttpRetrySystem::FRetryDomainsPtr RetryDomains = MakeShared<FHttpRetrySystem::FRetryDomains>(MoveTemp(AltDomains));
 	TSharedRef<IHttpRequest> HttpRequest = HttpRetryManager->CreateRequest(
 		1/*InRetryLimitCountOverride*/,
@@ -2071,13 +2083,18 @@ TEST_CASE_METHOD(FWaitUntilCompleteHttpFixture, "Retry immediately without lock
 		RetryDomains
 	);
 
-	HttpRequest->SetURL(UrlWithInvalidPortToTestConnectTimeout());
+	HttpRequest->SetURL(UrlToTestMethods());
 	HttpRequest->SetVerb(TEXT("GET"));
 	++ExpectingExtraCallbacks;
 	HttpRequest->OnRequestWillRetry().BindLambda([this](FHttpRequestPtr /*Request*/, FHttpResponsePtr /*Response*/, float LockoutPeriod) {
 		--ExpectingExtraCallbacks;
 		CHECK(LockoutPeriod == 0);
 	});
+	HttpRequest->OnProcessRequestComplete().BindLambda([](FHttpRequestPtr HttpRequest, FHttpResponsePtr HttpResponse, bool bSucceeded) {
+		CHECK(bSucceeded);
+		REQUIRE(HttpResponse != nullptr);
+		CHECK(HttpResponse->GetResponseCode() == 200);
+	});
 	HttpRequest->ProcessRequest();
 }
 
@@ -2110,7 +2127,8 @@ TEST_CASE_METHOD(FWaitUntilCompleteHttpFixture, "Optionally retry limit can be s
 	float TimeDiffTolerance = 0.0f;
 	SECTION("RetryLimitCountForConnectionErrorDefault:1 will be used so retries for connection error take less time")
 	{
-		HttpRequest->SetURL(UrlWithInvalidPortToTestConnectTimeout());
+		SetInvalidDomainToTestConnectTimeout();
+		HttpRequest->SetURL(UrlMockStatus(EHttpResponseCodes::TooManyRequests));
 
 		ExpectedTimeoutDuration = 6.0f; // each request will take 1s, 1st retry back off takes 4s
 		TimeDiffTolerance = 2 * UE_HTTP_CONNECTION_TIMEOUT_MAX_DEVIATION;
