Fix the potential crash: when trigger the mock failure and call the complete delegate, the delegate callback can reset the shared ptr reference of the http request while the request itself hasn't been added into http manager yet in PreProcess, so the instance get destroyed and it can't call StartTotalTimeoutTimer anymore;
Refactor: Remove the rollback code of CVarHttpSetGeneralFailureReasonFromCommonCode.

diff --git a/Engine/Source/Programs/WebTests/Private/TestHttp.cpp b/Engine/Source/Programs/WebTests/Private/TestHttp.cpp
--- a/Engine/Source/Programs/WebTests/Private/TestHttp.cpp
+++ b/Engine/Source/Programs/WebTests/Private/TestHttp.cpp
@@ -432,6 +432,7 @@ TEST_CASE_METHOD(FWaitUntilCompleteHttpFixture, "Can mock connection error throu
 		CHECK(HttpRequest->GetFailureReason() == EHttpFailureReason::ConnectionError);
 		const double DurationInSeconds  = FPlatformTime::Seconds() - StartTime;
 		CHECK(FMath::IsNearlyEqual(DurationInSeconds, ExpectedTimeoutDuration, HTTP_TIME_DIFF_TOLERANCE_OF_REQUEST));
+		CVarHttpUrlPatternsToMockFailure->Set(TEXT(""));
 	});
 	HttpRequest->ProcessRequest();
 }
@@ -447,10 +448,93 @@ TEST_CASE_METHOD(FWaitUntilCompleteHttpFixture, "Can mock response failure throu
 		CHECK(bSucceeded);
 		REQUIRE(HttpResponse != nullptr);
 		CHECK(HttpResponse->GetResponseCode() == 503);
+		CVarHttpUrlPatternsToMockFailure->Set(TEXT(""));
 	});
 	HttpRequest->ProcessRequest();
 }
 
+// Test the cases when http request wasn't added into http manager and completed, but user code reset the HttpRequest 
+// in the complete callback, so it will be destroyed because that's the only reference. 
+// Verify there is no more access from http code after the complete callback in such scenarios(FinishRequestNotInHttpManager)
+class FEarlyCompleteHttpFixture : public FWaitUntilCompleteHttpFixture
+{
+public:
+	~FEarlyCompleteHttpFixture()
+	{
+		while (HttpRequest.IsValid())
+		{
+			TickHttpManager();
+		}
+	}
+
+	TSharedPtr<IHttpRequest> HttpRequest;
+};
+
+TEST_CASE_METHOD(FEarlyCompleteHttpFixture, "Can delete failed request during complete callback for mock failures", HTTP_TAG)
+{
+	// When use retry manager, the request adapter will keep the shared ptr of the request impl, which can't reproduce the deletion
+	if (bRetryEnabled)
+	{
+		return;
+	}
+
+	DisableWarningsInThisTest();
+
+	CVarHttpUrlPatternsToMockFailure->Set(TEXT("epicgames.com->0 unrealengine.com->503"));
+
+	HttpRequest = CreateRequest();
+
+	SECTION("For mock failure with status code")
+	{
+		HttpRequest->SetURL(TEXT("https://www.unrealengine.com/"));
+	}
+	SECTION("For mock failure with connection error")
+	{
+		HttpModule->HttpConnectionTimeout = 2.0f;
+		HttpRequest->SetURL(TEXT("https://www.epicgames.com/"));
+	}
+
+	HttpRequest->OnProcessRequestComplete().BindLambda([this](FHttpRequestPtr InHttpRequest, FHttpResponsePtr HttpResponse, bool bSucceeded) {
+		CVarHttpUrlPatternsToMockFailure->Set(TEXT(""));
+		HttpRequest.Reset();
+	});
+
+	HttpRequest->ProcessRequest();
+}
+
+TEST_CASE_METHOD(FEarlyCompleteHttpFixture, "Can delete failed request during complete callback in pre process phase", HTTP_TAG)
+{
+	// When use retry manager, the request adapter will keep the shared ptr of the request impl, which can't reproduce the deletion
+	if (bRetryEnabled)
+	{
+		return;
+	}
+
+	DisableWarningsInThisTest();
+
+	HttpRequest = CreateRequest();
+
+	SECTION("For pre check failure when verb was set to empty")
+	{
+		HttpRequest->SetURL(TEXT("https://www.unrealengine.com/"));
+		HttpRequest->SetVerb(TEXT(""));
+	}
+	SECTION("For setup failure")
+	{
+		FString Filename = FString(FPlatformProcess::UserSettingsDir()) / TEXT("TestStreamUploadNonExist.dat");
+		HttpRequest->SetURL(UrlStreamUpload());
+		HttpRequest->SetVerb(TEXT("PUT"));
+		HttpRequest->SetHeader(TEXT("Content-Disposition"), TEXT("attachment;filename=TestStreamUploadNonExist.dat"));
+		HttpRequest->SetContentAsStreamedFile(Filename);
+	}
+
+	HttpRequest->OnProcessRequestComplete().BindLambda([this](FHttpRequestPtr InHttpRequest, FHttpResponsePtr HttpResponse, bool bSucceeded) {
+		HttpRequest.Reset();
+	});
+
+	HttpRequest->ProcessRequest();
+}
+
 TEST_CASE_METHOD(FWaitUntilCompleteHttpFixture, "Can complete successfully for different response codes", HTTP_TAG)
 {
 	TSharedRef<IHttpRequest> HttpRequest = CreateRequest();
@@ -1198,6 +1282,24 @@ TEST_CASE_METHOD(FWaitUntilCompleteHttpFixture, "Streaming http upload from file
 	HttpRequest->ProcessRequest();
 }
 
+TEST_CASE_METHOD(FWaitUntilCompleteHttpFixture, "Streaming http upload from invalid file by PUT will fail", HTTP_TAG)
+{
+	DisableWarningsInThisTest();
+	FString Filename = FString(FPlatformProcess::UserSettingsDir()) / TEXT("TestStreamUploadNonExist.dat");
+
+	TSharedRef<IHttpRequest> HttpRequest = CreateRequest();
+	HttpRequest->SetURL(UrlStreamUpload());
+	HttpRequest->SetVerb(TEXT("PUT"));
+	HttpRequest->SetHeader(TEXT("Content-Disposition"), TEXT("attachment;filename=TestStreamUploadNonExist.dat"));
+	HttpRequest->SetContentAsStreamedFile(Filename);
+	HttpRequest->OnProcessRequestComplete().BindLambda([Filename](FHttpRequestPtr HttpRequest, FHttpResponsePtr HttpResponse, bool bSucceeded) {
+		CHECK(!bSucceeded);
+		CHECK(HttpRequest->GetFailureReason() == EHttpFailureReason::Other);
+		CHECK(!HttpResponse);
+	});
+	HttpRequest->ProcessRequest();
+}
+
 TEST_CASE_METHOD(FWaitUntilCompleteHttpFixture, "Streaming uploading http request will re-open file when retry", HTTP_TAG)
 {
 	if (!bRetryEnabled)
diff --git a/Engine/Source/Runtime/Online/HTTP/Private/Apple/AppleHttp.cpp b/Engine/Source/Runtime/Online/HTTP/Private/Apple/AppleHttp.cpp
--- a/Engine/Source/Runtime/Online/HTTP/Private/Apple/AppleHttp.cpp
+++ b/Engine/Source/Runtime/Online/HTTP/Private/Apple/AppleHttp.cpp
@@ -11,8 +11,6 @@
 #include "Misc/App.h"
 #include "Misc/Base64.h"
 
-extern TAutoConsoleVariable<bool> CVarHttpSetGeneralFailureReasonFromCommonCode;
-
 /**
  * Class to hold data from delegate implementation notifications.
  */
@@ -888,15 +886,7 @@ void FAppleHttpRequest::FinishRequest()
 			ResponseCommon = nullptr;
 		}
 
-		if (CVarHttpSetGeneralFailureReasonFromCommonCode.GetValueOnAnyThread())
-		{
-			HandleRequestFailed(ResponseCommon);
-		}
-		else
-		{
-			SetStatus(EHttpRequestStatus::Failed);
-			OnProcessRequestComplete().ExecuteIfBound(SharedThis(this), ResponseCommon, bSucceeded);
-		}
+		HandleRequestFailed(ResponseCommon);
 	}
 }
 
diff --git a/Engine/Source/Runtime/Online/HTTP/Private/Curl/CurlHttp.cpp b/Engine/Source/Runtime/Online/HTTP/Private/Curl/CurlHttp.cpp
--- a/Engine/Source/Runtime/Online/HTTP/Private/Curl/CurlHttp.cpp
+++ b/Engine/Source/Runtime/Online/HTTP/Private/Curl/CurlHttp.cpp
@@ -19,8 +19,6 @@
 #include <openssl/ssl.h>
 #endif
 
-extern TAutoConsoleVariable<bool> CVarHttpSetGeneralFailureReasonFromCommonCode;
-
 TAutoConsoleVariable<bool> CVarHttpCurlSetResponseCodeWhenMarkAsCompleted(
 	TEXT("http.CurlSetResponseCodeWhenMarkAsCompleted"),
 	true,
@@ -1297,78 +1295,27 @@ void FCurlHttpRequest::FinishRequest()
 			}
 		}
 
-		if (CVarHttpSetGeneralFailureReasonFromCommonCode.GetValueOnAnyThread())
+		if (bCurlRequestCompleted)
 		{
-			if (bCurlRequestCompleted)
+			switch (CurlCompletionResult)
 			{
-				switch (CurlCompletionResult)
-				{
-				case CURLE_COULDNT_CONNECT:
-				case CURLE_OPERATION_TIMEDOUT:
-				case CURLE_COULDNT_RESOLVE_PROXY:
-				case CURLE_COULDNT_RESOLVE_HOST:
-				case CURLE_SSL_CONNECT_ERROR:
+			case CURLE_COULDNT_CONNECT:
+			case CURLE_OPERATION_TIMEDOUT:
+			case CURLE_COULDNT_RESOLVE_PROXY:
+			case CURLE_COULDNT_RESOLVE_HOST:
+			case CURLE_SSL_CONNECT_ERROR:
 #if WITH_CURL_XCURL
-				case CURLE_SEND_ERROR:
+			case CURLE_SEND_ERROR:
 #endif
-					// report these as connection errors (safe to retry)
-					SetFailureReason(EHttpFailureReason::ConnectionError);
-					break;
-				default:
-					break;
-				}
-			}
-
-			HandleRequestFailed(Response);
-		}
-		else
-		{
-			SetStatus(EHttpRequestStatus::Failed);
-
-			// Mark last request attempt as completed but failed
-			if (bCanceled)
-			{
-				SetFailureReason(EHttpFailureReason::Cancelled);
-			}
-			else if (bTimedOut)
-			{
-				SetFailureReason(EHttpFailureReason::TimedOut);
-			}
-			else if (bActivityTimedOut)
-			{
+				// report these as connection errors (safe to retry)
 				SetFailureReason(EHttpFailureReason::ConnectionError);
+				break;
+			default:
+				break;
 			}
-			else if (bCurlRequestCompleted)
-			{
-				switch (CurlCompletionResult)
-				{
-				case CURLE_COULDNT_CONNECT:
-				case CURLE_OPERATION_TIMEDOUT:
-				case CURLE_COULDNT_RESOLVE_PROXY:
-				case CURLE_COULDNT_RESOLVE_HOST:
-				case CURLE_SSL_CONNECT_ERROR:
-	#if WITH_CURL_XCURL
-				case CURLE_SEND_ERROR:
-	#endif
-					// report these as connection errors (safe to retry)
-					SetFailureReason(EHttpFailureReason::ConnectionError);
-					break;
-				default:
-					SetFailureReason(EHttpFailureReason::Other);
-				}
-			}
-			else
-			{
-				SetFailureReason(EHttpFailureReason::Other);
-			}
-
-			// Call delegate with failure
-			OnProcessRequestComplete().ExecuteIfBound(SharedThis(this), Response, false);
-
-			//Delegate needs to know about the errors -- so clear out Response (since connection failed) afterwards...
-			ResponseCommon = nullptr;
-			TotalBytesRead = 0;
 		}
+
+		HandleRequestFailed(Response);
 	}
 }
 
diff --git a/Engine/Source/Runtime/Online/HTTP/Private/GenericPlatform/HttpRequestCommon.cpp b/Engine/Source/Runtime/Online/HTTP/Private/GenericPlatform/HttpRequestCommon.cpp
--- a/Engine/Source/Runtime/Online/HTTP/Private/GenericPlatform/HttpRequestCommon.cpp
+++ b/Engine/Source/Runtime/Online/HTTP/Private/GenericPlatform/HttpRequestCommon.cpp
@@ -23,13 +23,6 @@ TAutoConsoleVariable<bool> CVarHttpLogJsonResponseOnly(
 
 }
 
-TAutoConsoleVariable<bool> CVarHttpSetGeneralFailureReasonFromCommonCode(
-	TEXT("http.SetGeneralFailureReasonFromCommonCode"),
-	true,
-	TEXT("Temporary hotfixable cvar: when enabled, set general failure reason from common code instead."),
-	ECVF_SaveForNextBoot
-);
-
 FHttpRequestCommon::FHttpRequestCommon()
 	: RequestStartTimeAbsoluteSeconds(FPlatformTime::Seconds())
 	, ActivityTimeoutAt(0.0)
@@ -109,7 +102,7 @@ bool FHttpRequestCommon::PreCheck() const
 	return true;
 }
 
-bool FHttpRequestCommon::WillTriggerMockFailure()
+bool FHttpRequestCommon::TriggerMockFailure()
 {
 	TOptional<int32> MockResponseCode = FHttpModule::Get().GetHttpManager().GetMockFailure(GetURL());
 	if (MockResponseCode.IsSet())
@@ -125,6 +118,10 @@ bool FHttpRequestCommon::WillTriggerMockFailure()
 					RequestPtr->FinishRequestNotInHttpManager();
 				}
 			}, HttpConnectionTimeout);
+
+			// Connect timeout mocking will trigger FinishRequest after a delay, still make sure total timeout 
+			// works when mocking connect timeout
+			StartTotalTimeoutTimer();
 		}
 		else
 		{
@@ -167,11 +164,8 @@ bool FHttpRequestCommon::PreProcess()
 		return false;
 	}
 
-	if (WillTriggerMockFailure())
+	if (TriggerMockFailure())
 	{
-		// Connect timeout mocking will trigger FinishRequest after a delay, still make sure total timeout 
-		// works when mocking connect timeout
-		StartTotalTimeoutTimer();
 		return false;
 	}
 
@@ -267,37 +261,35 @@ void FHttpRequestCommon::HandleRequestSucceed(TSharedPtr<IHttpResponse> InRespon
 
 	LogResponse(InResponse);
 
-	OnProcessRequestComplete().ExecuteIfBound(SharedThis(this), InResponse, true);
 	FHttpModule::Get().GetHttpManager().RecordStatTimeToConnect(ConnectTime);
+
+	OnProcessRequestComplete().ExecuteIfBound(SharedThis(this), InResponse, true);
 }
 
 void FHttpRequestCommon::HandleRequestFailed(TSharedPtr<IHttpResponse> InResponse)
 {
-	if (CVarHttpSetGeneralFailureReasonFromCommonCode.GetValueOnAnyThread())
+	if (FailureReason == EHttpFailureReason::None) // Failure reason was not set by platform, will set it here
 	{
-		if (FailureReason == EHttpFailureReason::None) // Failure reason was not set by platform, will set it here
+		if (bCanceled)
 		{
-			if (bCanceled)
-			{
-				SetFailureReason(EHttpFailureReason::Cancelled);
-			}
-			else if (bTimedOut)
-			{
-				SetFailureReason(EHttpFailureReason::TimedOut);
-			}
-			else if (!bUsePlatformActivityTimeout && bActivityTimedOut)
-			{
-				SetFailureReason(EHttpFailureReason::ConnectionError);
-			}
-			else
-			{
-				SetFailureReason(EHttpFailureReason::Other);
-			}
+			SetFailureReason(EHttpFailureReason::Cancelled);
+		}
+		else if (bTimedOut)
+		{
+			SetFailureReason(EHttpFailureReason::TimedOut);
+		}
+		else if (!bUsePlatformActivityTimeout && bActivityTimedOut)
+		{
+			SetFailureReason(EHttpFailureReason::ConnectionError);
+		}
+		else
+		{
+			SetFailureReason(EHttpFailureReason::Other);
 		}
-
-		SetStatus(EHttpRequestStatus::Failed);
 	}
 
+	SetStatus(EHttpRequestStatus::Failed);
+
 	OnProcessRequestComplete().ExecuteIfBound(SharedThis(this), InResponse, false);
 }
 
diff --git a/Engine/Source/Runtime/Online/HTTP/Public/GenericPlatform/HttpRequestCommon.h b/Engine/Source/Runtime/Online/HTTP/Public/GenericPlatform/HttpRequestCommon.h
--- a/Engine/Source/Runtime/Online/HTTP/Public/GenericPlatform/HttpRequestCommon.h
+++ b/Engine/Source/Runtime/Online/HTTP/Public/GenericPlatform/HttpRequestCommon.h
@@ -108,7 +108,7 @@ protected:
 	HTTP_API void CloseRequestPayloadDefaultImpl();
 
 	HTTP_API void LogResponse(const TSharedPtr<IHttpResponse>& InResponse);
-	HTTP_API bool WillTriggerMockFailure();
+	HTTP_API bool TriggerMockFailure();
 
 	HTTP_API void InitResponse();
 
