Remove the checking of the duplicated protocol, only keep the warning log.

diff --git a/Engine/Source/Programs/WebTests/Private/TestHttp.cpp b/Engine/Source/Programs/WebTests/Private/TestHttp.cpp
--- a/Engine/Source/Programs/WebTests/Private/TestHttp.cpp
+++ b/Engine/Source/Programs/WebTests/Private/TestHttp.cpp
@@ -223,15 +223,16 @@ TEST_CASE_METHOD(FHttpModuleTestFixture, "Shutdown http module without issue whe
 	HttpModule->GetHttpManager().Tick(0.0f);
 }
 
-TEST_CASE_METHOD(FHttpModuleTestFixture, "Duplicated scheme protocol will trigger check", HTTP_TAG)
+TEST_CASE_METHOD(FHttpModuleTestFixture, "Duplicated scheme protocol will trigger warning", HTTP_TAG)
 {
+	// TODO: find a way to catch expected warning log instead
 	DisableWarningsInThisTest();
 
 	TSharedRef<IHttpRequest> HttpRequest = CreateRequest();
 	CHECK(HttpRequest->GetVerb() == TEXT("GET"));
 
 	HttpRequest->SetURL(TEXT("https://https://www.epicgames.com"));
-	REQUIRE_CHECK_MSG("Duplicated protocol", HttpRequest->ProcessRequest());
+	HttpRequest->ProcessRequest();
 }
 
 TEST_CASE_METHOD(FHttpModuleTestFixture, "Shutdown http module without issue when there are ongoing streaming http requests with timeout.", HTTP_TAG)
diff --git a/Engine/Source/Runtime/Online/HTTP/Private/HttpManager.cpp b/Engine/Source/Runtime/Online/HTTP/Private/HttpManager.cpp
--- a/Engine/Source/Runtime/Online/HTTP/Private/HttpManager.cpp
+++ b/Engine/Source/Runtime/Online/HTTP/Private/HttpManager.cpp
@@ -248,7 +248,6 @@ static void CheckNoDuplicatedProtocol(FStringView UrlStringView)
 	if (FoundIndex != INDEX_NONE)
 	{
 		FStringView Domain = UrlStringView.RightChop(FoundIndex + Delimiter.Len());
-		checkf(!Domain.StartsWith(UrlStringView.Left(FoundIndex + Delimiter.Len())), TEXT("Duplicated protocol in url: %.*s"), UrlStringView.Len(), UrlStringView.GetData());
 		UE_CLOG(Domain.Contains(Delimiter), LogHttp, Warning, TEXT("Duplicated delimiter '://' found in url %.*s, is that intended?"), UrlStringView.Len(), UrlStringView.GetData());
 	}
 }
