Fix the potential crash if Connect websockets when shutdown, the reason: when calling Connect, it could find the existing security info on the platform, and push the task to queue in http manager, in the next frame after shutting down websocket module, http module can still call flush tick which will call the delegate, in the delegate it will try to access the data of websockets manager which has been deleted.

Solution: Simply check if the manager singleton is still there before accessing the existing security info in the delegate. Will change it to use safer version of delegate for the websockets manager in the future.

Added corresponding test case to repro and cover it.

diff --git a/Engine/Source/Programs/WebTests/Private/TestWebSockets.cpp b/Engine/Source/Programs/WebTests/Private/TestWebSockets.cpp
--- a/Engine/Source/Programs/WebTests/Private/TestWebSockets.cpp
+++ b/Engine/Source/Programs/WebTests/Private/TestWebSockets.cpp
@@ -4,6 +4,7 @@
 #include "HAL/IConsoleManager.h"
 #include "HAL/PlatformProcess.h"
 #include "Http.h"
+#include "HttpManager.h"
 #include "Misc/CommandLine.h"
 #include "WebSocketsLog.h"
 #include "WebSocketsModule.h"
@@ -54,6 +55,9 @@ public:
 		WebSocketsModuleInterface->ShutdownModule();
 		delete WebSocketsModuleInterface;
 
+		// Simulate the flush tick when shutdown, to make sure tasks added by FHttpModule::Get().GetHttpManager().AddGameThreadTask won't crash
+		HttpModule->GetHttpManager().Tick(0.0);
+
 		IModuleInterface* HttpModuleInterface = HttpModule;
 		HttpModuleInterface->ShutdownModule();
 		delete HttpModuleInterface;
@@ -98,14 +102,19 @@ public:
 		RunUntilQuitRequested();
 	}
 
+	void SimulateEngineTick()
+	{
+		FTSBackgroundableTicker::GetCoreTicker().Tick(TickFrequency);
+		FTSTicker::GetCoreTicker().Tick(TickFrequency);
+		FPlatformProcess::Sleep(TickFrequency);
+		HttpModule->GetHttpManager().Tick(TickFrequency);
+	}
+
 	void RunUntilQuitRequested()
 	{
 		while (!bQuitRequested)
 		{
-			FTSBackgroundableTicker::GetCoreTicker().Tick(TickFrequency);
-			FTSTicker::GetCoreTicker().Tick(TickFrequency);
-
-			FPlatformProcess::Sleep(TickFrequency);
+			SimulateEngineTick();
 		}
 	}
 
@@ -158,3 +167,49 @@ TEST_CASE_METHOD(FRunUntilQuitRequestedFixture, "WebSockets module can shut down
 
 	WebSocket->Connect();
 }
+
+class FConnectWhenShutdownFixture : public FRunUntilQuitRequestedFixture
+{
+public:
+	FConnectWhenShutdownFixture()
+	{
+		WebSocket = WebSocketsModule->CreateWebSocket(FString::Format(TEXT("{0}/echo/"), { *UrlWebSocketsTests() }));
+	}
+
+	~FConnectWhenShutdownFixture()
+	{
+		// Waiting until closed, then the Connect call can actually launch without early return
+		RunUntilQuitRequested();
+
+		WebSocket->Connect();
+	}
+
+	TSharedPtr<IWebSocket> WebSocket;
+};
+
+TEST_CASE_METHOD(FConnectWhenShutdownFixture, "WebSockets can call connect when shutdown", WEBSOCKETS_TAG)
+{
+	DisableWarningsInThisTest();
+
+	WebSocket->OnConnected().AddLambda([this](){
+		WebSocket->Send(TEXT("hi websockets tests"));
+	});
+
+	WebSocket->OnMessage().AddLambda([this](const FString& MessageString) {
+		CHECK(MessageString == TEXT("hi websockets tests"));
+		WebSocket->Close();
+		bSucceed = true;
+	});
+
+	WebSocket->OnClosed().AddLambda([this](int32 /* StatusCode */, const FString& /* Reason */, bool /* bWasClean */){
+		CHECK(bSucceed);
+		bQuitRequested = true;
+	});
+
+	WebSocket->OnConnectionError().AddLambda([this](const FString& /* Error */){
+		CHECK(false);
+		bQuitRequested = true;
+	});
+
+	WebSocket->Connect();
+}
