Add early wake up support for http event loop multi wait implementation.

diff --git a/Engine/Source/Programs/WebTests/Private/TestHttp.cpp b/Engine/Source/Programs/WebTests/Private/TestHttp.cpp
--- a/Engine/Source/Programs/WebTests/Private/TestHttp.cpp
+++ b/Engine/Source/Programs/WebTests/Private/TestHttp.cpp
@@ -40,6 +40,7 @@
 extern TAutoConsoleVariable<bool> CVarHttpInsecureProtocolEnabled;
 extern TAutoConsoleVariable<bool> CVarHttpRetrySystemNonGameThreadSupportEnabled;
 extern TAutoConsoleVariable<int32> CVarHttpMaxConcurrentRequests;
+extern TAutoConsoleVariable<int32> CVarHttpEventLoopEnableChance;
 extern TAutoConsoleVariable<FString> CVarHttpUrlPatternsToMockFailure;
 
 class FMockHttpModule : public FHttpModule
@@ -113,9 +114,12 @@ public:
 		, WebServerHttpPort(8000)
 		, bRunHeavyTests(false)
 		, bRetryEnabled(true)
+		, bEventLoopEnabled(true)
 	{
 		ParseSettingsFromCommandLine();
 
+		CVarHttpEventLoopEnableChance->Set(bEventLoopEnabled ? 100 : 0);
+
 		bRetryEnabled &= CVarHttpRetrySystemNonGameThreadSupportEnabled.GetValueOnAnyThread();
 
 		InitModule();
@@ -157,6 +161,7 @@ public:
 		FParse::Value(FCommandLine::Get(), TEXT("web_server_ip="), WebServerIp);
 		FParse::Bool(FCommandLine::Get(), TEXT("run_heavy_tests="), bRunHeavyTests);
 		FParse::Bool(FCommandLine::Get(), TEXT("retry_enabled="), bRetryEnabled);
+		FParse::Bool(FCommandLine::Get(), TEXT("eventloop_enabled="), bEventLoopEnabled);
 		FParse::Value(FCommandLine::Get(), TEXT("web_server_unix_socket="), WebServerUnixSocket);
 	}
 
@@ -192,6 +197,7 @@ public:
 	FMockHttpModule* HttpModule = nullptr;
 	bool bRunHeavyTests;
 	bool bRetryEnabled;
+	bool bEventLoopEnabled;
 	FHttpTestLogLevelInitializer HttpTestLogLevelInitializer;
 	TSharedPtr<FMockRetryManager> HttpRetryManager;
 };
@@ -926,15 +932,98 @@ TEST_CASE_METHOD(FWaitUntilCompleteHttpFixture, "Can run parallel stream downloa
 	}
 }
 
-TEST_CASE_METHOD(FWaitUntilCompleteHttpFixture, "Can download big file exceeds 32 bits", HTTP_TAG)
+TEST_CASE_METHOD(FWaitUntilCompleteHttpFixture, "Can early wake up when adding new request", HTTP_TAG)
+{
+	const double StartTime = FPlatformTime::Seconds();
+
+	// First request will wait for a long time before receiving any update
+	TSharedRef<IHttpRequest> HttpRequestSlow = FHttpModule::Get().CreateRequest();
+	HttpRequestSlow ->SetURL(UrlStreamDownload(1, HTTP_TEST_TIMEOUT_CHUNK_SIZE, /*ChunkLatency=*/3));
+	HttpRequestSlow->SetDelegateThreadPolicy(EHttpRequestDelegateThreadPolicy::CompleteOnHttpThread);
+
+	HttpRequestSlow->OnProcessRequestComplete().BindLambda([](FHttpRequestPtr HttpRequestSlow, FHttpResponsePtr HttpResponse, bool bSucceeded)
+	{
+		CHECK(!bSucceeded);
+	});
+	HttpRequestSlow->ProcessRequest();
+
+	FPlatformProcess::Sleep(1);
+
+	// Second request will complete immediately
+	TSharedRef<IHttpRequest> HttpRequestFast = FHttpModule::Get().CreateRequest();
+	HttpRequestFast->SetURL(UrlStreamDownload(3, HTTP_TEST_TIMEOUT_CHUNK_SIZE, /*ChunkLatency=*/ 0));
+	HttpRequestFast->SetDelegateThreadPolicy(EHttpRequestDelegateThreadPolicy::CompleteOnHttpThread);
+	HttpRequestFast->OnProcessRequestComplete().BindLambda([StartTime, HttpRequestSlow](FHttpRequestPtr HttpRequestFast, FHttpResponsePtr HttpResponse, bool bSucceeded)
+	{
+		CHECK(bSucceeded);
+		REQUIRE(HttpResponse != nullptr);
+		CHECK(HttpResponse->GetResponseCode() == 200);
+
+		const double DurationInSeconds  = FPlatformTime::Seconds() - StartTime;
+		CHECK(DurationInSeconds < 2.0f);
+
+		CHECK(HttpRequestSlow->GetStatus() == EHttpRequestStatus::Processing);
+		HttpRequestSlow->CancelRequest();
+	});
+	HttpRequestFast->ProcessRequest();
+}
+
+TEST_CASE_METHOD(FWaitUntilCompleteHttpFixture, "Can run a lot of small download requests with good performance", HTTP_TAG)
 {
 	if (!bRunHeavyTests)
 	{
 		return;
 	}
 
+	const double StartTime = FPlatformTime::Seconds();
+
+	int32 NumRequests = 10000;
+	FEvent* Event = FPlatformProcess::GetSynchEventFromPool(true);
+
+	std::atomic<int32> RunningRequests = 0;
+	for (int32 i = 0; i < NumRequests; ++i)
+	{
+		TSharedRef<IHttpRequest> HttpRequest = FHttpModule::Get().CreateRequest();
+		HttpRequest->SetURL(UrlStreamDownload(3, HTTP_TEST_TIMEOUT_CHUNK_SIZE, /*ChunkLatency=*/ 0));
+		HttpRequest->SetDelegateThreadPolicy(EHttpRequestDelegateThreadPolicy::CompleteOnHttpThread);
+
+		HttpRequest->OnProcessRequestComplete().BindLambda([Event, &RunningRequests](FHttpRequestPtr HttpRequest, FHttpResponsePtr HttpResponse, bool bSucceeded)
+		{
+			CHECK(bSucceeded);
+			REQUIRE(HttpResponse != nullptr);
+			CHECK(HttpResponse->GetResponseCode() == 200);
+			--RunningRequests;
+		});
+
+		// Add requests constantly to validate early wake up performance instead of adding all at beginning
+		while (RunningRequests > 20)
+		{
+			FPlatformProcess::Sleep(0.0f);
+		}
+
+		++RunningRequests;
+		HttpRequest->ProcessRequest();
+	}
+
+	WaitUntilAllHttpRequestsComplete();
+
+	const double DurationInSeconds  = FPlatformTime::Seconds() - StartTime;
+
+	UE_LOG(LogHttp, Log, TEXT("%d small download requests spent %.2fs to complete!"), NumRequests, DurationInSeconds);
+}
+
+TEST_CASE_METHOD(FWaitUntilCompleteHttpFixture, "Can download big file exceeds 32 bits with good performance", HTTP_TAG)
+{
+	if (!bRunHeavyTests)
+	{
+		return;
+	}
+
+	const double StartTime = FPlatformTime::Seconds();
+
 	// 5 * 1024 * 1024 * 1024 BYTES = 5368709120 BYTES = 5 GB
-	uint64 Chunks = 5 * 1024;
+	uint64 DataSizeInGB = 5;
+	uint64 Chunks = DataSizeInGB * 1024;
 	uint64 ChunkSize = 1024 * 1024;
 
 	TSharedRef<IHttpRequest> HttpRequest = CreateRequest();
@@ -957,6 +1046,12 @@ TEST_CASE_METHOD(FWaitUntilCompleteHttpFixture, "Can download big file exceeds 3
 		CHECK(HttpResponse->GetResponseCode() == 200);
 	});
 	HttpRequest->ProcessRequest();
+
+	WaitUntilAllHttpRequestsComplete();
+
+	const double DurationInSeconds  = FPlatformTime::Seconds() - StartTime;
+
+	UE_LOG(LogHttp, Log, TEXT("%d GB of big download request spent %.2fs to complete!"), DataSizeInGB, DurationInSeconds);
 }
 
 TEST_CASE_METHOD(FWaitUntilCompleteHttpFixture, "Streaming http upload from memory", HTTP_TAG)
diff --git a/Engine/Source/Runtime/Online/Experimental/EventLoop/EventLoop.Build.cs b/Engine/Source/Runtime/Online/Experimental/EventLoop/EventLoop.Build.cs
--- a/Engine/Source/Runtime/Online/Experimental/EventLoop/EventLoop.Build.cs
+++ b/Engine/Source/Runtime/Online/Experimental/EventLoop/EventLoop.Build.cs
@@ -19,6 +19,8 @@ public class EventLoop : ModuleRules
 
 		PublicDefinitions.Add(string.Format("HAS_EVENTLOOP_PLATFORM_SOCKET_IMPLEMENTATION={0}", bHasPlatformSocketImplementation ? 1 : 0));
 		PublicDefinitions.Add(string.Format("HAS_EVENTLOOP_PLATFORM_BSD_SOCKET_HEADER={0}", bHasPlatformBSDSocketHeader ? 1 : 0));
+		PublicDefinitions.Add(string.Format("HAS_PLATFORM_EVENTLOOP_NOTIFIER_EVENTFD={0}", bHasPlatformEventLoopNotifierEventFd ? 1 : 0));
+		PublicDefinitions.Add(string.Format("HAS_EVENTLOOP_NOTIFIER_EVENTFD={0}", bHasEventLoopNotifierEventFd ? 1 : 0));
 	}
 
 	public virtual bool bHasPlatformSocketImplementation
@@ -38,4 +40,7 @@ public class EventLoop : ModuleRules
 			// Other platforms may override this property.
 		}
 	}
+
+	public virtual bool bHasPlatformEventLoopNotifierEventFd { get { return false; } }
+	public virtual bool bHasEventLoopNotifierEventFd { get { return Target.IsInPlatformGroup(UnrealPlatformGroup.Linux); } }
 }
diff --git a/Engine/Source/Runtime/Online/Experimental/EventLoop/Private/EventLoop/BSDSocket/BSDSocketTypesPrivate.h b/Engine/Source/Runtime/Online/Experimental/EventLoop/Private/EventLoop/BSDSocket/BSDSocketTypesPrivate.h
--- a/Engine/Source/Runtime/Online/Experimental/EventLoop/Private/EventLoop/BSDSocket/BSDSocketTypesPrivate.h
+++ b/Engine/Source/Runtime/Online/Experimental/EventLoop/Private/EventLoop/BSDSocket/BSDSocketTypesPrivate.h
@@ -42,6 +42,10 @@ typedef int32 SOCKLEN;
 #include <netinet/tcp.h>
 #endif // PLATFORM_HAS_BSD_SOCKET_FEATURE_NODELAY
 
+#if HAS_EVENTLOOP_NOTIFIER_EVENTFD && !HAS_PLATFORM_EVENTLOOP_NOTIFIER_EVENTFD
+#include <sys/eventfd.h>
+#endif // HAS_EVENTLOOP_NOTIFIER_EVENTFD
+
 #define ioctlsocket ioctl
 #endif // HAS_EVENTLOOP_PLATFORM_BSD_SOCKET_HEADER
 
diff --git a/Engine/Source/Runtime/Online/Experimental/EventLoop/Private/EventLoop/BSDSocket/EventLoopNotifierEventFd.cpp b/Engine/Source/Runtime/Online/Experimental/EventLoop/Private/EventLoop/BSDSocket/EventLoopNotifierEventFd.cpp
new file mode 100644
--- /dev/null
+++ b/Engine/Source/Runtime/Online/Experimental/EventLoop/Private/EventLoop/BSDSocket/EventLoopNotifierEventFd.cpp
@@ -0,0 +1,40 @@
+// Copyright Epic Games, Inc. All Rights Reserved.
+
+#include "EventLoop/BSDSocket/EventLoopNotifierEventFd.h"
+
+#if HAS_EVENTLOOP_NOTIFIER_EVENTFD
+#include "BSDSocketTypesPrivate.h"
+
+bool FEventLoopNotifierEventFd::Init()
+{
+	FileDescriptorRead = eventfd(0, EFD_NONBLOCK);
+	return FileDescriptorRead != INVALID_SOCKET;
+}
+
+void FEventLoopNotifierEventFd::Shutdown()
+{
+	if (FileDescriptorRead != INVALID_SOCKET)
+	{
+		close(FileDescriptorRead);
+		FileDescriptorRead = INVALID_SOCKET;
+	}
+}
+
+void FEventLoopNotifierEventFd::Notify()
+{
+	eventfd_t Value = 1;
+	if (eventfd_write(FileDescriptorRead, Value) == SOCKET_RESULT_FAILED)
+	{
+		int Err = errno;
+		check(Err == EAGAIN);
+	}
+}
+
+void FEventLoopNotifierEventFd::Clear()
+{
+	eventfd_t Value = 0;
+	ensure(eventfd_read(FileDescriptorRead, &Value) != SOCKET_RESULT_FAILED);
+}
+
+#endif // HAS_EVENTLOOP_NOTIFIER_EVENTFD
+
diff --git a/Engine/Source/Runtime/Online/Experimental/EventLoop/Private/EventLoop/BSDSocket/EventLoopNotifierSocket.cpp b/Engine/Source/Runtime/Online/Experimental/EventLoop/Private/EventLoop/BSDSocket/EventLoopNotifierSocket.cpp
new file mode 100644
--- /dev/null
+++ b/Engine/Source/Runtime/Online/Experimental/EventLoop/Private/EventLoop/BSDSocket/EventLoopNotifierSocket.cpp
@@ -0,0 +1,142 @@
+// Copyright Epic Games, Inc. All Rights Reserved.
+
+#include "EventLoop/BSDSocket/EventLoopNotifierSocket.h"
+#include "BSDSocketTypesPrivate.h"
+
+#if PLATFORM_HAS_BSD_SOCKETS && !PLATFORM_HAS_BSD_SOCKET_FEATURE_WINSOCKETS
+
+namespace UE::EventLoopSocketNotifierInternal
+{
+
+static void SetSocketNonBlock(int32 FileDescriptor)
+{
+	int32 Flags = fcntl(FileDescriptor, F_GETFL, 0);
+	fcntl(FileDescriptor, F_SETFL, O_NONBLOCK | Flags );
+}
+
+}
+
+bool FEventLoopNotifierSocket::Init()
+{
+	// TODO: Validate and handle the case for socket recreation in case of suspend and resume
+
+	// TODO: Investigate if it's possible to use pipe or IPPROTO_RAW, or socketpair instead
+	FileDescriptorListener = socket(AF_INET, SOCK_STREAM, IPPROTO_TCP);
+	sockaddr_in AddrIn;
+	memset(&AddrIn, 0, sizeof(AddrIn));
+	AddrIn.sin_family = AF_INET;
+	AddrIn.sin_addr.s_addr = htonl(INADDR_LOOPBACK);
+	AddrIn.sin_port = 0;
+
+	if (bind(FileDescriptorListener, reinterpret_cast<sockaddr*>(&AddrIn), sizeof(AddrIn)) == SOCKET_RESULT_FAILED)
+	{
+		close(FileDescriptorListener);
+		FileDescriptorListener = INVALID_SOCKET;
+
+		return false;
+	}
+
+	if (listen(FileDescriptorListener, /*backlog*/1) == SOCKET_RESULT_FAILED)
+	{
+		close(FileDescriptorListener);
+		FileDescriptorListener = INVALID_SOCKET;
+
+		return false;
+	}
+
+	UE::EventLoopSocketNotifierInternal::SetSocketNonBlock(FileDescriptorListener);
+
+
+
+
+	// TODO: Set the write buffer to 1
+	sockaddr Addr;
+	memset(&Addr, 0, sizeof(Addr));
+
+	socklen_t AddrLen = sizeof(Addr);
+	getsockname(FileDescriptorListener, &Addr, &AddrLen);
+
+	FileDescriptorRead = socket(AF_INET, SOCK_STREAM, IPPROTO_TCP);
+	if (connect(FileDescriptorRead, &Addr, sizeof(Addr)) == SOCKET_RESULT_FAILED)
+	{
+		close(FileDescriptorListener);
+		FileDescriptorListener = INVALID_SOCKET;
+
+		close(FileDescriptorRead);
+		FileDescriptorRead = INVALID_SOCKET;
+		return false;
+	}
+
+	UE::EventLoopSocketNotifierInternal::SetSocketNonBlock(FileDescriptorRead);
+
+
+
+
+	// TODO: Set the read buffer to 1
+	FileDescriptorWrite = accept(FileDescriptorListener, NULL, NULL);
+	if (FileDescriptorWrite == INVALID_SOCKET)
+	{
+		close(FileDescriptorListener);
+		FileDescriptorListener = INVALID_SOCKET;
+
+		close(FileDescriptorRead);
+		FileDescriptorRead = INVALID_SOCKET;
+		return false;
+	}
+
+	UE::EventLoopSocketNotifierInternal::SetSocketNonBlock(FileDescriptorWrite);
+
+
+
+
+	return true;
+}
+
+void FEventLoopNotifierSocket::Shutdown()
+{
+	if (FileDescriptorWrite != INVALID_SOCKET)
+	{
+		close(FileDescriptorWrite);
+		FileDescriptorWrite = INVALID_SOCKET;
+	}
+
+	if (FileDescriptorRead != INVALID_SOCKET)
+	{
+		close(FileDescriptorRead);
+		FileDescriptorRead = INVALID_SOCKET;
+	}
+
+	if (FileDescriptorListener != INVALID_SOCKET)
+	{
+		close(FileDescriptorListener);
+		FileDescriptorListener = INVALID_SOCKET;
+	}
+}
+
+void FEventLoopNotifierSocket::Notify()
+{
+	if (FileDescriptorWrite != INVALID_SOCKET)
+	{
+		// TODO: A failure to write may not be critical. Its possible that the receive buffer is just full. 
+		// The notifier just cares that the socket is set to a readable state.
+		if (write(FileDescriptorWrite, "1", 1) == SOCKET_RESULT_FAILED)
+		{
+			Shutdown();
+		}
+	}
+}
+
+void FEventLoopNotifierSocket::Clear()
+{
+	if (FileDescriptorRead != INVALID_SOCKET)
+	{
+		char Buffer[1];
+		if (read(FileDescriptorRead, &Buffer[0], 1) == SOCKET_RESULT_FAILED)
+		{
+			Shutdown();
+		}
+	}
+}
+
+#endif // PLATFORM_HAS_BSD_SOCKETS && !PLATFORM_HAS_BSD_SOCKET_FEATURE_WINSOCKETS
+
diff --git a/Engine/Source/Runtime/Online/Experimental/EventLoop/Public/EventLoop/BSDSocket/BSDSocketTypes.h b/Engine/Source/Runtime/Online/Experimental/EventLoop/Public/EventLoop/BSDSocket/BSDSocketTypes.h
--- a/Engine/Source/Runtime/Online/Experimental/EventLoop/Public/EventLoop/BSDSocket/BSDSocketTypes.h
+++ b/Engine/Source/Runtime/Online/Experimental/EventLoop/Public/EventLoop/BSDSocket/BSDSocketTypes.h
@@ -14,4 +14,6 @@
 	#define INVALID_SOCKET -1
 #endif // PLATFORM_HAS_BSD_SOCKET_FEATURE_WINSOCKETS
 
+#define SOCKET_RESULT_FAILED -1
+
 #endif // PLATFORM_HAS_BSD_SOCKETS
\ No newline at end of file
diff --git a/Engine/Source/Runtime/Online/Experimental/EventLoop/Public/EventLoop/BSDSocket/EventLoopNotifier.h b/Engine/Source/Runtime/Online/Experimental/EventLoop/Public/EventLoop/BSDSocket/EventLoopNotifier.h
new file mode 100644
--- /dev/null
+++ b/Engine/Source/Runtime/Online/Experimental/EventLoop/Public/EventLoop/BSDSocket/EventLoopNotifier.h
@@ -0,0 +1,13 @@
+// Copyright Epic Games, Inc. All Rights Reserved.
+
+#pragma once
+
+#include "CoreTypes.h"
+
+#if HAS_EVENTLOOP_NOTIFIER_EVENTFD
+#include "EventLoop/BSDSocket/EventLoopNotifierEventFd.h"
+using FEventLoopNotifier = FEventLoopNotifierEventFd;
+#elif PLATFORM_HAS_BSD_SOCKETS && !PLATFORM_HAS_BSD_SOCKET_FEATURE_WINSOCKETS
+#include "EventLoop/BSDSocket/EventLoopNotifierSocket.h"
+using FEventLoopNotifier = FEventLoopNotifierSocket;
+#endif
diff --git a/Engine/Source/Runtime/Online/Experimental/EventLoop/Public/EventLoop/BSDSocket/EventLoopNotifierEventFd.h b/Engine/Source/Runtime/Online/Experimental/EventLoop/Public/EventLoop/BSDSocket/EventLoopNotifierEventFd.h
new file mode 100644
--- /dev/null
+++ b/Engine/Source/Runtime/Online/Experimental/EventLoop/Public/EventLoop/BSDSocket/EventLoopNotifierEventFd.h
@@ -0,0 +1,25 @@
+// Copyright Epic Games, Inc. All Rights Reserved.
+
+#pragma once
+
+#include "BSDSocketTypes.h"
+#include "CoreTypes.h"
+
+#if HAS_EVENTLOOP_NOTIFIER_EVENTFD
+
+class FEventLoopNotifierEventFd
+{
+public:
+	bool Init();
+	void Shutdown();
+
+	void Notify();
+	void Clear();
+
+	int32 GetFileDescriptorRead() const { return FileDescriptorRead; }
+
+private:
+	int32 FileDescriptorRead = INVALID_SOCKET;
+};
+
+#endif // HAS_EVENTLOOP_NOTIFIER_EVENTFD
diff --git a/Engine/Source/Runtime/Online/Experimental/EventLoop/Public/EventLoop/BSDSocket/EventLoopNotifierSocket.h b/Engine/Source/Runtime/Online/Experimental/EventLoop/Public/EventLoop/BSDSocket/EventLoopNotifierSocket.h
new file mode 100644
--- /dev/null
+++ b/Engine/Source/Runtime/Online/Experimental/EventLoop/Public/EventLoop/BSDSocket/EventLoopNotifierSocket.h
@@ -0,0 +1,28 @@
+// Copyright Epic Games, Inc. All Rights Reserved.
+
+#pragma once
+
+#include "BSDSocketTypes.h"
+#include "CoreTypes.h"
+
+#if PLATFORM_HAS_BSD_SOCKETS && !PLATFORM_HAS_BSD_SOCKET_FEATURE_WINSOCKETS
+
+class FEventLoopNotifierSocket
+{
+public:
+	bool Init();
+	void Shutdown();
+
+	void Notify();
+	void Clear();
+
+	int32 GetFileDescriptorRead() const { return FileDescriptorRead; }
+
+private:
+	int32 FileDescriptorListener = INVALID_SOCKET;
+	int32 FileDescriptorWrite = INVALID_SOCKET;
+	int32 FileDescriptorRead = INVALID_SOCKET;
+};
+
+#endif // PLATFORM_HAS_BSD_SOCKETS && !PLATFORM_HAS_BSD_SOCKET_FEATURE_WINSOCKETS
+
diff --git a/Engine/Source/Runtime/Online/HTTP/Private/Curl/CurlMultiWaitEventLoopHttpThread.cpp b/Engine/Source/Runtime/Online/HTTP/Private/Curl/CurlMultiWaitEventLoopHttpThread.cpp
--- a/Engine/Source/Runtime/Online/HTTP/Private/Curl/CurlMultiWaitEventLoopHttpThread.cpp
+++ b/Engine/Source/Runtime/Online/HTTP/Private/Curl/CurlMultiWaitEventLoopHttpThread.cpp
@@ -26,16 +26,23 @@ FCurlMultiWaitIOManager::FCurlMultiWaitIOManager(UE::EventLoop::IEventLoop& InEv
 
 bool FCurlMultiWaitIOManager::Init()
 {
+	if (!ensure(EventLoopNotifier.Init()))
+	{
+		UE_LOG(LogHttp, Warning, TEXT("Socket notifier init failed, multi-wait eventloop performance may degrade!"));
+	}
+
 	check(Params.ProcessCurlRequests);
 	return true;
 }
 
 void FCurlMultiWaitIOManager::Shutdown()
 {
+	EventLoopNotifier.Shutdown();
 }
 
 void FCurlMultiWaitIOManager::Notify()
 {
+	EventLoopNotifier.Notify();
 }
 
 void FCurlMultiWaitIOManager::Poll(FTimespan WaitTime)
@@ -50,30 +57,49 @@ void FCurlMultiWaitIOManager::Poll(FTimespan WaitTime)
 
 	Params.ProcessCurlRequests();
 
-	// The maximum wait time affects how quickly new requests are started / canceled.
-	// Todo: add early wakeup handling.
-	const FTimespan MaxWaitTime = FTimespan::FromMilliseconds(10);
-	FTimespan CurrentWaitTime = FMath::Min(WaitTime, MaxWaitTime);
-
-	// A call to curl_multi_wait will return immediately with no FDs set when a timer is fired or
-	// when there are no sockets being handled. When it returns early the first time assume that a
-	// timer has fired. On the second iteration assume that no FDs are set and instead sleep for
-	// the duration.
-	if (EmptySequentialWaitCount > 1)
+	int32 FileDescriptorRead = EventLoopNotifier.GetFileDescriptorRead();
+	if (FileDescriptorRead != INVALID_SOCKET)
 	{
-		QUICK_SCOPE_CYCLE_COUNTER(STAT_FCurlMultiWaitIOManager_Poll_Sleep);
-		FPlatformProcess::SleepNoStats(CurrentWaitTime.GetTotalSeconds());
-		EmptySequentialWaitCount = 0;
+		struct curl_waitfd EarlyWakeUpWait;
+		EarlyWakeUpWait.fd = FileDescriptorRead;
+		EarlyWakeUpWait.events = CURL_WAIT_POLLIN; /* wait for input */
+		EarlyWakeUpWait.revents = 0; /* clear it */
+
+		const int64 TimeoutMs = FMath::RoundToPositiveInfinity(WaitTime.GetTotalMilliseconds());
+		int NumFDs = 0;
+
+		CheckMultiCodeOk(curl_multi_wait(Params.MultiHandle, &EarlyWakeUpWait, 1, TimeoutMs, &NumFDs));
+		if (EarlyWakeUpWait.revents > 0)
+		{
+			EventLoopNotifier.Clear();
+		}
 	}
 	else
 	{
-		QUICK_SCOPE_CYCLE_COUNTER(STAT_FCurlMultiWaitIOManager_Poll_MultiWait);
-		const int64 TimeoutMs = FMath::RoundToPositiveInfinity(WaitTime.GetTotalMilliseconds());
-		int NumFDs = 0;
-		CheckMultiCodeOk(curl_multi_wait(Params.MultiHandle, nullptr, 0, TimeoutMs, &NumFDs));
+		// The maximum wait time affects how quickly new requests are started / canceled.
+		const FTimespan MaxWaitTime = FTimespan::FromMilliseconds(10);
+		FTimespan CurrentWaitTime = FMath::Min(WaitTime, MaxWaitTime);
+
+		// A call to curl_multi_wait will return immediately with no FDs set when a timer is fired or
+		// when there are no sockets being handled. When it returns early the first time assume that a
+		// timer has fired. On the second iteration assume that no FDs are set and instead sleep for
+		// the duration.
+		if (EmptySequentialWaitCount > 1)
+		{
+			QUICK_SCOPE_CYCLE_COUNTER(STAT_FCurlMultiWaitIOManager_Poll_Sleep);
+			FPlatformProcess::SleepNoStats(CurrentWaitTime.GetTotalSeconds());
+			EmptySequentialWaitCount = 0;
+		}
+		else
+		{
+			QUICK_SCOPE_CYCLE_COUNTER(STAT_FCurlMultiWaitIOManager_Poll_MultiWait);
+			const int64 TimeoutMs = FMath::RoundToPositiveInfinity(WaitTime.GetTotalMilliseconds());
+			int NumFDs = 0;
+			CheckMultiCodeOk(curl_multi_wait(Params.MultiHandle, nullptr, 0, TimeoutMs, &NumFDs));
 
-		// Track how many sequential calls had no FDs set.
-		EmptySequentialWaitCount = (NumFDs != 0) ? 0 : ++EmptySequentialWaitCount;
+			// Track how many sequential calls had no FDs set.
+			EmptySequentialWaitCount = (NumFDs != 0) ? 0 : ++EmptySequentialWaitCount;
+		}
 	}
 }
 
diff --git a/Engine/Source/Runtime/Online/HTTP/Private/Curl/CurlMultiWaitEventLoopHttpThread.h b/Engine/Source/Runtime/Online/HTTP/Private/Curl/CurlMultiWaitEventLoopHttpThread.h
--- a/Engine/Source/Runtime/Online/HTTP/Private/Curl/CurlMultiWaitEventLoopHttpThread.h
+++ b/Engine/Source/Runtime/Online/HTTP/Private/Curl/CurlMultiWaitEventLoopHttpThread.h
@@ -9,6 +9,7 @@
 
 #include "EventLoopHttpThread.h"
 #include "EventLoop/EventLoop.h"
+#include "EventLoop/BSDSocket/EventLoopNotifier.h"
 #include "Templates/Function.h"
 
 #if PLATFORM_MICROSOFT
@@ -67,6 +68,7 @@ private:
 	UE::EventLoop::IEventLoop& EventLoop;
 	FParams Params;
 	int32 EmptySequentialWaitCount;
+	FEventLoopNotifier EventLoopNotifier;
 };
 
 class FCurlMultiWaitEventLoopHttpThread
