Clear the OnChangedDelegate of CVarHttpUrlPatternsToDisableFailedLog in FHttpManager::Shutdown;
Add test cases for FHttpManager::OnBeforeFork and failed http reuqest log disabling.

diff --git a/Engine/Source/Programs/WebTests/Private/TestHttp.cpp b/Engine/Source/Programs/WebTests/Private/TestHttp.cpp
--- a/Engine/Source/Programs/WebTests/Private/TestHttp.cpp
+++ b/Engine/Source/Programs/WebTests/Private/TestHttp.cpp
@@ -13,6 +13,7 @@
 #include "HttpRouteHandle.h"
 #include "IHttpRouter.h"
 #include "HttpServerModule.h"
+#include "Logging/LogScopedVerbosityOverride.h"
 #include "Misc/CommandLine.h"
 #include "TestHarness.h"
 #include "TestMacros/Assertions.h"
@@ -42,6 +43,7 @@ extern TAutoConsoleVariable<bool> CVarHttpRetrySystemNonGameThreadSupportEnabled
 extern TAutoConsoleVariable<int32> CVarHttpMaxConcurrentRequests;
 extern TAutoConsoleVariable<int32> CVarHttpEventLoopEnableChance;
 extern TAutoConsoleVariable<FString> CVarHttpUrlPatternsToMockFailure;
+extern TAutoConsoleVariable<FString> CVarHttpUrlPatternsToDisableFailedLog;
 
 class FMockHttpModule : public FHttpModule
 {
@@ -2550,6 +2552,54 @@ TEST_CASE_METHOD(FWaitUntilCompleteHttpFixture, "Flush while activity timeout sh
 	HttpModule->GetHttpManager().Flush(EHttpFlushReason::FullFlush);
 }
 
+TEST_CASE_METHOD(FWaitUntilCompleteHttpFixture, "OnBeforeFork can clear http requests", HTTP_TAG)
+{
+	// Warning: Sleeping 0.500s to wait for 1 outstanding Http Requests.
+	DisableWarningsInThisTest();
+
+	{
+		TSharedPtr<IHttpRequest> HttpRequest = CreateRequest();
+		HttpRequest->SetURL(UrlStreamDownload(3/*Chunks*/, HTTP_TEST_TIMEOUT_CHUNK_SIZE, 0/*ChunkLatency*/));
+		HttpRequest->OnProcessRequestComplete().BindLambda([](FHttpRequestPtr HttpRequest, FHttpResponsePtr HttpResponse, bool bSucceeded) {
+			CHECK(bSucceeded);
+		});
+		HttpRequest->ProcessRequest();
+	}
+
+	CHECK(OngoingRequests == 1);
+
+	HttpModule->GetHttpManager().OnBeforeFork();
+
+	CHECK(OngoingRequests == 0);
+
+	HttpModule->GetHttpManager().OnAfterFork();
+}
+
+TEST_CASE_METHOD(FWaitUntilQuitFromTestFixture, "Can disable failed log through CVar", HTTP_TAG)
+{
+	// Only disable warning in this function scope for pre check warning log, but not disable warning when complete
+	LOG_SCOPE_VERBOSITY_OVERRIDE(LogHttp, ELogVerbosity::Error);
+
+	CVarHttpUrlPatternsToDisableFailedLog->Set(TEXT("unrealengine.com"));
+
+	TSharedRef<IHttpRequest> HttpRequest = CreateRequest();
+	HttpRequest->SetURL(TEXT("https://www.unrealengine.com/"));
+	HttpRequest->SetVerb(TEXT("")); // Pre check will fail
+	// Make sure when complete it's not in this function scope
+	HttpRequest->SetDelegateThreadPolicy(EHttpRequestDelegateThreadPolicy::CompleteOnHttpThread);
+
+	HttpRequest->OnProcessRequestComplete().BindLambda([this](FHttpRequestPtr HttpRequest, FHttpResponsePtr HttpResponse, bool bSucceeded) {
+		bQuitRequested = true;
+		CHECK(!bSucceeded);
+		HttpModule->GetHttpManager().AddGameThreadTask([]()
+		{
+			// Have to clear CVar out of the test function scope, but also need to be in game thread.
+			CVarHttpUrlPatternsToDisableFailedLog->Set(TEXT(""));
+		});
+	});
+	HttpRequest->ProcessRequest();
+}
+
 #if UE_HTTP_SUPPORT_LOCAL_SERVER
 
 TEST_CASE_METHOD(FWaitUntilCompleteHttpFixture, "Scheme besides http and https can work if allowed by settings", HTTP_TAG)
diff --git a/Engine/Source/Runtime/Online/HTTP/Private/HttpManager.cpp b/Engine/Source/Runtime/Online/HTTP/Private/HttpManager.cpp
--- a/Engine/Source/Runtime/Online/HTTP/Private/HttpManager.cpp
+++ b/Engine/Source/Runtime/Online/HTTP/Private/HttpManager.cpp
@@ -130,6 +130,7 @@ void FHttpManager::Initialize()
 void FHttpManager::Shutdown()
 {
 	CVarHttpUrlPatternsToLogResponse.AsVariable()->OnChangedDelegate().Clear();
+	CVarHttpUrlPatternsToDisableFailedLog.AsVariable()->OnChangedDelegate().Clear();
 	CVarHttpUrlPatternsToMockFailure.AsVariable()->OnChangedDelegate().Clear();
 
 	{
