Move the common code of calling HandleRequestSucceed and HandleRequestFailed from http implementations into new OnFinishRequest method in HttpRequestCommon;
In non-shipping build, trigger the check when user code calling Flush during http request complete delegate.

diff --git a/Engine/Source/Programs/WebTests/Private/TestHttp.cpp b/Engine/Source/Programs/WebTests/Private/TestHttp.cpp
--- a/Engine/Source/Programs/WebTests/Private/TestHttp.cpp
+++ b/Engine/Source/Programs/WebTests/Private/TestHttp.cpp
@@ -2575,6 +2575,47 @@ TEST_CASE_METHOD(FWaitUntilCompleteHttpFixture, "OnBeforeFork can clear http req
 	HttpModule->GetHttpManager().OnAfterFork();
 }
 
+TEST_CASE_METHOD(FWaitUntilCompleteHttpFixture, "Will trigger check if flush during http complete delegate in game thread", HTTP_TAG)
+{
+	TSharedPtr<IHttpRequest> HttpRequest = CreateRequest();
+	HttpRequest->SetURL(UrlToTestMethods());
+	HttpRequest->OnProcessRequestComplete().BindLambda([this](FHttpRequestPtr HttpRequest, FHttpResponsePtr HttpResponse, bool bSucceeded) {
+		REQUIRE_CHECK_MSG("It's not supported to call Flush during http request complete delegate!", HttpModule->GetHttpManager().Flush(EHttpFlushReason::FullFlush));
+	});
+	HttpRequest->ProcessRequest();
+}
+
+TEST_CASE_METHOD(FWaitThreadedHttpFixture, "Will trigger check if flush during http complete delegate in http thread", HTTP_TAG)
+{
+	ThreadedHttpRunnable.OnRunFromThread().BindLambda([this]() {
+		TSharedPtr<IHttpRequest> HttpRequest = CreateRequest();
+		HttpRequest->SetURL(UrlToTestMethods());
+		HttpRequest->OnProcessRequestComplete().BindLambda([this](FHttpRequestPtr HttpRequest, FHttpResponsePtr HttpResponse, bool bSucceeded) {
+			REQUIRE_CHECK_MSG("It's not supported to call Flush during http request complete delegate!", HttpModule->GetHttpManager().Flush(EHttpFlushReason::FullFlush));
+		});
+		HttpRequest->ProcessRequest();
+	});
+
+	ThreadedHttpRunnable.StartTestHttpThread(/*bBlockGameThread=*/false);
+}
+
+TEST_CASE_METHOD(FWaitThreadedHttpFixture, "Will NOT trigger check if flush from another thread during http complete delegate", HTTP_TAG)
+{
+	TSharedRef<IHttpRequest> HttpRequest = CreateRequest();
+	HttpRequest->SetURL(UrlToTestMethods());
+	HttpRequest->OnProcessRequestComplete().BindLambda([this](FHttpRequestPtr HttpRequest, FHttpResponsePtr HttpResponse, bool bSucceeded) {
+		CHECK(bSucceeded);
+		FPlatformProcess::Sleep(1.0);
+	});
+	HttpRequest->ProcessRequest();
+
+	ThreadedHttpRunnable.OnRunFromThread().BindLambda([this]() {
+		FPlatformProcess::Sleep(0.5); // Make sure the request reached OnProcessRequestComplete
+		HttpModule->GetHttpManager().Flush(EHttpFlushReason::FullFlush);
+	});
+	ThreadedHttpRunnable.StartTestHttpThread(/*bBlockGameThread=*/false);
+}
+
 TEST_CASE_METHOD(FWaitUntilQuitFromTestFixture, "Can disable failed log through CVar", HTTP_TAG)
 {
 	// Only disable warning in this function scope for pre check warning log, but not disable warning when complete
diff --git a/Engine/Source/Runtime/Online/HTTP/Private/Apple/AppleHttp.cpp b/Engine/Source/Runtime/Online/HTTP/Private/Apple/AppleHttp.cpp
--- a/Engine/Source/Runtime/Online/HTTP/Private/Apple/AppleHttp.cpp
+++ b/Engine/Source/Runtime/Online/HTTP/Private/Apple/AppleHttp.cpp
@@ -859,11 +859,7 @@ void FAppleHttpRequest::FinishRequest()
 	TSharedPtr<FAppleHttpResponse> Response = StaticCastSharedPtr<FAppleHttpResponse>(ResponseCommon);
 	bool bSucceeded = (Response && Response->GetStatusFromDelegate() == EHttpRequestStatus::Succeeded);
 
-	if (bSucceeded)
-	{
-		HandleRequestSucceed(Response);
-	}
-	else
+	if (!bSucceeded)
 	{
 		if (FailureReason == EHttpFailureReason::None) // FailureReason could have been set by FHttpRequestCommon::WillTriggerMockFailure
 		{
@@ -894,9 +890,9 @@ void FAppleHttpRequest::FinishRequest()
 		{
 			ResponseCommon = nullptr;
 		}
-
-		HandleRequestFailed(ResponseCommon);
 	}
+
+	OnFinishRequest(bSucceeded);
 }
 
 void FAppleHttpRequest::CleanupRequest()
diff --git a/Engine/Source/Runtime/Online/HTTP/Private/Curl/CurlHttp.cpp b/Engine/Source/Runtime/Online/HTTP/Private/Curl/CurlHttp.cpp
--- a/Engine/Source/Runtime/Online/HTTP/Private/Curl/CurlHttp.cpp
+++ b/Engine/Source/Runtime/Online/HTTP/Private/Curl/CurlHttp.cpp
@@ -1249,11 +1249,8 @@ void FCurlHttpRequest::FinishRequest()
 		Response->bIsReady = true;
 	}
 
-	if (Response.IsValid() && Response->bSucceeded)
-	{
-		HandleRequestSucceed(Response);
-	}
-	else
+	bool bSucceeded = Response.IsValid() && Response->bSucceeded;
+	if (!bSucceeded)
 	{
 		if (CurlAddToMultiResult != CURLM_OK)
 		{
@@ -1296,9 +1293,9 @@ void FCurlHttpRequest::FinishRequest()
 				break;
 			}
 		}
-
-		HandleRequestFailed(Response);
 	}
+
+	OnFinishRequest(bSucceeded);
 }
 
 void FCurlHttpRequest::CleanupRequest()
diff --git a/Engine/Source/Runtime/Online/HTTP/Private/GenericPlatform/HttpRequestCommon.cpp b/Engine/Source/Runtime/Online/HTTP/Private/GenericPlatform/HttpRequestCommon.cpp
--- a/Engine/Source/Runtime/Online/HTTP/Private/GenericPlatform/HttpRequestCommon.cpp
+++ b/Engine/Source/Runtime/Online/HTTP/Private/GenericPlatform/HttpRequestCommon.cpp
@@ -254,18 +254,16 @@ void FHttpRequestCommon::SetOption(const FName Option, const FString& OptionValu
 	Options.Add(Option, OptionValue);
 }
 
-void FHttpRequestCommon::HandleRequestSucceed(TSharedPtr<IHttpResponse> InResponse)
+void FHttpRequestCommon::HandleRequestSucceed()
 {
 	SetStatus(EHttpRequestStatus::Succeeded);
 
-	LogResponse(InResponse);
+	LogResponse(ResponseCommon);
 
 	FHttpModule::Get().GetHttpManager().RecordStatTimeToConnect(ConnectTime);
-
-	OnProcessRequestComplete().ExecuteIfBound(SharedThis(this), InResponse, true);
 }
 
-void FHttpRequestCommon::HandleRequestFailed(TSharedPtr<IHttpResponse> InResponse)
+void FHttpRequestCommon::HandleRequestFailed()
 {
 	if (FailureReason == EHttpFailureReason::None) // Failure reason was not set by platform, will set it here
 	{
@@ -290,8 +288,6 @@ void FHttpRequestCommon::HandleRequestFailed(TSharedPtr<IHttpResponse> InRespons
 	SetStatus(EHttpRequestStatus::Failed);
 
 	LogFailure();
-
-	OnProcessRequestComplete().ExecuteIfBound(SharedThis(this), InResponse, false);
 }
 
 #define UE_HTTP_LOG_AS_WARNING_IF(Condition, Format, ...) \
@@ -787,4 +783,20 @@ void FHttpRequestCommon::LogResponse(const TSharedPtr<IHttpResponse>& InResponse
 
 		StartPos = EndPos + 1;
 	}
-}
\ No newline at end of file
+}
+
+void FHttpRequestCommon::OnFinishRequest(bool bSucceeded)
+{
+	// TODO: Move more code from impl into this common function
+
+	if (bSucceeded)
+	{
+		HandleRequestSucceed();
+	}
+	else
+	{
+		HandleRequestFailed();
+	}
+
+	OnProcessRequestComplete().ExecuteIfBound(SharedThis(this), ResponseCommon, bSucceeded);
+}
diff --git a/Engine/Source/Runtime/Online/HTTP/Private/GenericPlatform/HttpRequestCommon.h b/Engine/Source/Runtime/Online/HTTP/Private/GenericPlatform/HttpRequestCommon.h
--- a/Engine/Source/Runtime/Online/HTTP/Private/GenericPlatform/HttpRequestCommon.h
+++ b/Engine/Source/Runtime/Online/HTTP/Private/GenericPlatform/HttpRequestCommon.h
@@ -80,8 +80,8 @@ protected:
 	 */
 	HTTP_API void FinishRequestNotInHttpManager();
 
-	HTTP_API void HandleRequestSucceed(TSharedPtr<IHttpResponse> InResponse);
-	HTTP_API void HandleRequestFailed(TSharedPtr<IHttpResponse> InResponse);
+	HTTP_API void HandleRequestSucceed();
+	HTTP_API void HandleRequestFailed();
 
 	HTTP_API void StartActivityTimeoutTimer();
 	HTTP_API void StartActivityTimeoutTimerBy(double DelayToTrigger);
@@ -120,6 +120,8 @@ protected:
 	// The function for each impl to make sure when mocking response code, the internal state is valid for finishing request successfully
 	HTTP_API virtual void MockResponseData() = 0;
 
+	void OnFinishRequest(bool bSucceeded);
+
 protected:
 	/** Current status of request being processed */
 	EHttpRequestStatus::Type CompletionStatus = EHttpRequestStatus::NotStarted;
diff --git a/Engine/Source/Runtime/Online/HTTP/Private/HttpManager.cpp b/Engine/Source/Runtime/Online/HTTP/Private/HttpManager.cpp
--- a/Engine/Source/Runtime/Online/HTTP/Private/HttpManager.cpp
+++ b/Engine/Source/Runtime/Online/HTTP/Private/HttpManager.cpp
@@ -363,6 +363,7 @@ void FHttpManager::Flush(EHttpFlushReason FlushReason)
 {
 	QUICK_SCOPE_CYCLE_COUNTER(STAT_FHttpManager_Flush);
 
+	checkf(!IsCurrentThreadCompletingRequest(), TEXT("It's not supported to call Flush during http request complete delegate!"));
 	checkf(FlushReason != EHttpFlushReason::Shutdown || !HasAnyBoundDelegate(), TEXT("Use Shutdown() instead of Flush(EHttpFlushReason::Shutdown) directly."));
 	
 	// This variable is set to indicate that flush is happening.
@@ -527,8 +528,12 @@ bool FHttpManager::Tick(float DeltaSeconds)
 
 			if (CompletedRequest->GetDelegateThreadPolicy() == EHttpRequestDelegateThreadPolicy::CompleteOnGameThread)
 			{
+				MarkCurrentThreadCompletingRequest(true);
 				CompletedRequest->FinishRequest();
+				// Can't move BroadcastHttpRequestCompleted to CompletedRequest->FinishRequest because 
+				// FinishRequest can be called before adding into HttpManager when pre-check failed
 				BroadcastHttpRequestCompleted(CompletedRequestRef);
+				MarkCurrentThreadCompletingRequest(false);
 			}
 		}
 	}
@@ -768,3 +773,25 @@ TOptional<int32> FHttpManager::GetMockFailure(FStringView Url) const
 	return Result;
 }
 
+void FHttpManager::MarkCurrentThreadCompletingRequest(bool bCompleting)
+{
+#if DO_CHECK
+	const FScopeLock CacheLock(&ThreadsCompletingRequestCriticalSection);
+	bool& ThreadMark = ThreadsCompletingRequest.FindOrAdd(FPlatformTLS::GetCurrentThreadId());
+	ThreadMark = bCompleting;
+#endif
+}
+
+#if DO_CHECK
+bool FHttpManager::IsCurrentThreadCompletingRequest() const
+{
+	const FScopeLock CacheLock(&ThreadsCompletingRequestCriticalSection);
+	const bool* ThreadMark = ThreadsCompletingRequest.Find(FPlatformTLS::GetCurrentThreadId());
+	if (ThreadMark == nullptr)
+	{
+		return false;
+	}
+
+	return *ThreadMark;
+}
+#endif
diff --git a/Engine/Source/Runtime/Online/HTTP/Private/HttpThread.cpp b/Engine/Source/Runtime/Online/HTTP/Private/HttpThread.cpp
--- a/Engine/Source/Runtime/Online/HTTP/Private/HttpThread.cpp
+++ b/Engine/Source/Runtime/Online/HTTP/Private/HttpThread.cpp
@@ -371,8 +371,11 @@ void FHttpThreadBase::FinishRequestsFromHttpThreadWithCallbacks(TArray<FHttpRequ
 
 			if (Request->GetDelegateThreadPolicy() == EHttpRequestDelegateThreadPolicy::CompleteOnHttpThread)
 			{
+				FHttpManager& HttpManager = FHttpModule::Get().GetHttpManager();
+				HttpManager.MarkCurrentThreadCompletingRequest(true);
 				Request->FinishRequest();
-				FHttpModule::Get().GetHttpManager().BroadcastHttpRequestCompleted(Request->AsShared());
+				HttpManager.BroadcastHttpRequestCompleted(Request->AsShared());
+				HttpManager.MarkCurrentThreadCompletingRequest(false);
 			}
 
 			CompletedThreadedRequests.Enqueue(Request);
diff --git a/Engine/Source/Runtime/Online/HTTP/Public/HttpManager.h b/Engine/Source/Runtime/Online/HTTP/Public/HttpManager.h
--- a/Engine/Source/Runtime/Online/HTTP/Public/HttpManager.h
+++ b/Engine/Source/Runtime/Online/HTTP/Public/HttpManager.h
@@ -415,6 +415,11 @@ protected:
 	TMap<FString, int32> UrlPatternsToMockFailure;
 	mutable FCriticalSection UrlPatternsToMockFailureCriticalSection;
 
+#if DO_CHECK
+	mutable FCriticalSection ThreadsCompletingRequestCriticalSection;
+	TMap<uint32, bool> ThreadsCompletingRequest;
+#endif
+
 PACKAGE_SCOPE:
 
 	/** Used to lock access to add/remove/find requests */
@@ -454,4 +459,9 @@ PACKAGE_SCOPE:
 	HTTP_API bool ShouldDisableFailedLog(FStringView Url) const;
 
 	HTTP_API TOptional<int32> GetMockFailure(FStringView Url) const;
+
+	HTTP_API void MarkCurrentThreadCompletingRequest(bool bCompleting);
+#if DO_CHECK
+	HTTP_API bool IsCurrentThreadCompletingRequest() const;
+#endif
 };
