Add gzip support for xCurl;
Add test case for gzip support in WebTests.

diff --git a/Engine/Source/Programs/WebTests/Private/TestHttp.cpp b/Engine/Source/Programs/WebTests/Private/TestHttp.cpp
--- a/Engine/Source/Programs/WebTests/Private/TestHttp.cpp
+++ b/Engine/Source/Programs/WebTests/Private/TestHttp.cpp
@@ -470,6 +470,36 @@ TEST_CASE_METHOD(FWaitUntilCompleteHttpFixture, "Can mock response failure throu
 	HttpRequest->ProcessRequest();
 }
 
+TEST_CASE_METHOD(FWaitUntilCompleteHttpFixture, "Can process request with gzip support", HTTP_TAG)
+{
+	uint32 DataLength = 1025;
+	uint32 RepeatAt = 9;
+
+	TSharedRef<IHttpRequest> HttpRequest = CreateRequest();
+	HttpRequest->SetVerb(TEXT("GET"));
+	HttpRequest->SetURL(FString::Format(TEXT("{0}/get_data_gzip/{1}/{2}/"), { *UrlHttpTests(), DataLength, RepeatAt }));
+	HttpRequest->SetHeader(TEXT("Accept-Encoding"), TEXT("gzip"));
+	HttpRequest->OnProcessRequestComplete().BindLambda([DataLength, RepeatAt](FHttpRequestPtr HttpRequest, FHttpResponsePtr HttpResponse, bool bSucceeded) {
+		CHECK(bSucceeded);
+		REQUIRE(HttpResponse != nullptr);
+		CHECK(HttpResponse->GetResponseCode() == 200);
+
+		const TArray<uint8>& Content = HttpResponse->GetContent();
+		CHECK(Content.Num() == DataLength);
+
+		bool bAllMatch = true;
+
+		// Make sure the data got read is in good state
+		for (int32 i = 0; i < Content.Num(); ++i)
+		{
+			bAllMatch &= (Content[i] == '0' + (i % RepeatAt));
+		}
+
+		CHECK(bAllMatch);
+	});
+	HttpRequest->ProcessRequest();
+}
+
 // Test the cases when http request wasn't added into http manager and completed, but user code reset the HttpRequest 
 // in the complete callback, so it will be destroyed because that's the only reference. 
 // Verify there is no more access from http code after the complete callback in such scenarios(FinishRequestNotInHttpManager)
diff --git a/Engine/Source/Programs/WebTestsServer/httptests/urls.py b/Engine/Source/Programs/WebTestsServer/httptests/urls.py
--- a/Engine/Source/Programs/WebTestsServer/httptests/urls.py
+++ b/Engine/Source/Programs/WebTestsServer/httptests/urls.py
@@ -6,6 +6,7 @@ urlpatterns = [
     path("httptests/methods", views.methods, name="methods"),
     path("httptests/query_with_params/", views.query_with_params, name="query_with_params"),
     path("httptests/get_data_without_chunks/<int:bytes_number>/<int:repeat_at>/", views.get_data_without_chunks, name="get_data_without_chunks"),
+    path("httptests/get_data_gzip/<int:bytes_number>/<int:repeat_at>/", views.get_data_gzip, name="get_data_gzip"),
     path("httptests/streaming_download/<int:chunks>/<int:chunk_size>/<int:chunk_latency>/", views.streaming_download, name="streaming_download"),
     path("httptests/streaming_upload_post", views.streaming_upload_post, name="streaming_upload_post"),
     path("httptests/streaming_upload_put", views.streaming_upload_put, name="streaming_upload_put"),
diff --git a/Engine/Source/Programs/WebTestsServer/httptests/views.py b/Engine/Source/Programs/WebTestsServer/httptests/views.py
--- a/Engine/Source/Programs/WebTestsServer/httptests/views.py
+++ b/Engine/Source/Programs/WebTestsServer/httptests/views.py
@@ -5,6 +5,7 @@ from django.http.response import JsonResponse
 from django.http import JsonResponse
 from django.http import HttpResponse
 from django.http import StreamingHttpResponse
+from django.views.decorators.gzip import gzip_page
 from rest_framework.decorators import api_view, parser_classes
 from rest_framework.parsers import FileUploadParser
 import time
@@ -23,12 +24,21 @@ def query_with_params(request):
     var_str = request.GET.get('var_str')
     return JsonResponse({'var_int' : var_int, 'var_str' : var_str}, content_type="application/json")
 
-@api_view(['GET'])
-def get_data_without_chunks(request, bytes_number, repeat_at):
+def generate_data(bytes_number, repeat_at):
+    assert bytes_number > 0
     assert repeat_at > 0
     assert repeat_at <= 10
-    data = [i%repeat_at for i in range(bytes_number)]
-    return HttpResponse(data)
+    return [i%repeat_at for i in range(bytes_number)]
+
+@api_view(['GET'])
+def get_data_without_chunks(request, bytes_number, repeat_at):
+    return HttpResponse(generate_data(bytes_number, repeat_at))
+
+@api_view(['GET'])
+@gzip_page
+def get_data_gzip(request, bytes_number, repeat_at):
+    assert request.headers.get('Accept-Encoding') == 'gzip'
+    return HttpResponse(generate_data(bytes_number, repeat_at))
 
 @api_view(['GET'])
 def streaming_download(request, chunks, chunk_size, chunk_latency):
diff --git a/Engine/Source/Runtime/Online/HTTP/Private/Curl/CurlHttp.cpp b/Engine/Source/Runtime/Online/HTTP/Private/Curl/CurlHttp.cpp
--- a/Engine/Source/Runtime/Online/HTTP/Private/Curl/CurlHttp.cpp
+++ b/Engine/Source/Runtime/Online/HTTP/Private/Curl/CurlHttp.cpp
@@ -953,7 +953,12 @@ bool FCurlHttpRequest::SetupRequestHttpThread()
 		// Empty string here tells Curl to list all supported encodings, allowing servers to send compressed content.
 		if (FCurlHttpManager::CurlRequestOptions.bAcceptCompressedContent)
 		{
+#if WITH_CURL_XCURL
+			// xCurl doesn't seem to support passing empty string to enable all supported built-in compressions
+			curl_easy_setopt(EasyHandle, CURLOPT_ACCEPT_ENCODING, "gzip");
+#else
 			curl_easy_setopt(EasyHandle, CURLOPT_ACCEPT_ENCODING, "");
+#endif
 		}
 
 		{
