Seems retrying with connection failure is now supported on all platforms. Getting rid of that macro.

diff --git a/Engine/Source/Programs/WebTests/Private/TestHttp.cpp b/Engine/Source/Programs/WebTests/Private/TestHttp.cpp
--- a/Engine/Source/Programs/WebTests/Private/TestHttp.cpp
+++ b/Engine/Source/Programs/WebTests/Private/TestHttp.cpp
@@ -1573,7 +1573,6 @@ TEST_CASE_METHOD(FWaitUntilQuitFromTestFixture, "Http request can be reused when
 	HttpRequest->ProcessRequest();
 }
 
-#if UE_HTTP_CONNECTION_TIMEOUT_SUPPORT_RETRY
 TEST_CASE_METHOD(FWaitUntilQuitFromTestFixture, "Make sure connection time out can work well for 2nd same http request", HTTP_TAG)
 {
 	DisableWarningsInThisTest();
@@ -1600,7 +1599,6 @@ TEST_CASE_METHOD(FWaitUntilQuitFromTestFixture, "Make sure connection time out c
 	});
 	HttpRequest->ProcessRequest();
 }
-#endif // UE_HTTP_CONNECTION_TIMEOUT_SUPPORT_RETRY
 
 // Response shared ptr should be able to be kept by user code and valid to access without http request
 class FValidateResponseDependencyFixture : public FWaitUntilCompleteHttpFixture
@@ -2100,7 +2098,6 @@ TEST_CASE_METHOD(FWaitThreadedHttpFixture, "Cancel http request with ProcessRequ
 	ThreadedHttpRunnable.StartTestHttpThread(true/*bBlockGameThread*/);
 }
 
-#if UE_HTTP_CONNECTION_TIMEOUT_SUPPORT_RETRY
 TEST_CASE_METHOD(FWaitUntilCompleteHttpFixture, "Cancel http request connect before timeout", HTTP_TAG)
 {
 	DisableWarningsInThisTest();
@@ -2130,7 +2127,6 @@ TEST_CASE_METHOD(FWaitUntilCompleteHttpFixture, "Cancel http request connect bef
 	HttpRequest->CancelRequest();
 	HttpRequest->CancelRequest(); // Duplicated calls to CancelRequest should be fine
 }
-#endif
 
 TEST_CASE_METHOD(FWaitUntilCompleteHttpFixture, "Retry respect Retry-After header in response", HTTP_TAG)
 {
@@ -2423,7 +2419,6 @@ TEST_CASE_METHOD(FWaitUntilCompleteHttpFixture, "Retry will add default domain t
 	HttpRequest->ProcessRequest();
 }
 
-#if UE_HTTP_CONNECTION_TIMEOUT_SUPPORT_RETRY
 TEST_CASE_METHOD(FWaitUntilCompleteHttpFixture, "Optionally retry limit can be set differently for connection error", HTTP_TAG)
 {
 	if (!bRetryEnabled)
@@ -2474,7 +2469,6 @@ TEST_CASE_METHOD(FWaitUntilCompleteHttpFixture, "Optionally retry limit can be s
 	});
 	HttpRequest->ProcessRequest();
 }
-#endif // UE_HTTP_CONNECTION_TIMEOUT_SUPPORT_RETRY
 
 TEST_CASE_METHOD(FWaitUntilCompleteHttpFixture, "Retry fallback with exponential lock out if there is no Retry-After header", HTTP_TAG)
 {
diff --git a/Engine/Source/Runtime/Online/HTTP/HTTP.Build.cs b/Engine/Source/Runtime/Online/HTTP/HTTP.Build.cs
--- a/Engine/Source/Runtime/Online/HTTP/HTTP.Build.cs
+++ b/Engine/Source/Runtime/Online/HTTP/HTTP.Build.cs
@@ -35,7 +35,6 @@ public class HTTP : ModuleRules
 
 	protected virtual bool bPlatformSupportsCurlMultiWait { get { return false; } }
 	protected virtual bool bPlatformSupportsCurlQuickExit { get { return !bPlatformSupportsXCurl; } }
-	protected virtual bool bPlatformConnectionTimeoutSupportRetry { get { return true; } }
 	protected virtual bool bPlatformSupportsLocalHttpServer 
 	{ 
 		get 
@@ -158,7 +157,6 @@ public class HTTP : ModuleRules
 			PlatformConnectionTimeoutMaxDeviation = 1.5f;
 		}
 		PublicDefinitions.Add("UE_HTTP_CONNECTION_TIMEOUT_MAX_DEVIATION=" + PlatformConnectionTimeoutMaxDeviation);
-		PublicDefinitions.Add("UE_HTTP_CONNECTION_TIMEOUT_SUPPORT_RETRY=" + (bPlatformConnectionTimeoutSupportRetry ? "1" : "0"));
 		PublicDefinitions.Add("UE_HTTP_ACTIVITY_TIMER_START_AFTER_RECEIVED_DATA=" + ((bPlatformSupportsXCurl || Target.IsInPlatformGroup(UnrealPlatformGroup.Apple)) ? "1" : "0"));
 		PublicDefinitions.Add("UE_HTTP_SUPPORT_LOCAL_SERVER=" + (bPlatformSupportsLocalHttpServer ? "1" : "0"));
 		PublicDefinitions.Add("UE_HTTP_SUPPORT_UNIX_SOCKET=" + (bPlatformSupportsUnixSockets ? "1" : "0"));
