Fix the issue: The same http request can be added into http manager again before removed, when the request with CompleteOnHttpThread delegate policy got retried;
Get rid of CVarHttpRetrySystemNonGameThreadSupportEnabled which is no longer used;
Get rid of CVarHttpAddDomainOfDefaultUrlToRetryDomainsEnabled which is no longer used.

diff --git a/Engine/Source/Programs/WebTests/Private/TestHttp.cpp b/Engine/Source/Programs/WebTests/Private/TestHttp.cpp
--- a/Engine/Source/Programs/WebTests/Private/TestHttp.cpp
+++ b/Engine/Source/Programs/WebTests/Private/TestHttp.cpp
@@ -39,7 +39,6 @@
 #define HTTP_TEST_TIMEOUT_CHUNK_SIZE 16*1024 // Use a big chunk size so it triggers data received callback in time on all platforms
 
 extern TAutoConsoleVariable<bool> CVarHttpInsecureProtocolEnabled;
-extern TAutoConsoleVariable<bool> CVarHttpRetrySystemNonGameThreadSupportEnabled;
 extern TAutoConsoleVariable<int32> CVarHttpMaxConcurrentRequests;
 extern TAutoConsoleVariable<int32> CVarHttpEventLoopEnableChance;
 extern TAutoConsoleVariable<FString> CVarHttpUrlPatternsToMockFailure;
@@ -122,8 +121,6 @@ public:
 
 		CVarHttpEventLoopEnableChance->Set(bEventLoopEnabled ? 100 : 0);
 
-		bRetryEnabled &= CVarHttpRetrySystemNonGameThreadSupportEnabled.GetValueOnAnyThread();
-
 		InitModule();
 
 		CVarHttpInsecureProtocolEnabled->Set(true);
@@ -2317,7 +2314,6 @@ TEST_CASE_METHOD(FWaitUntilCompleteHttpFixture, "Retry immediately without lock
 	SetInvalidDomainToTestConnectTimeout();
 
 	FString InvalidDomain = UrlDomain();
-	AltDomains.Add(InvalidDomain);
 	AltDomains.Add(CorrectDomain);
 	FHttpRetrySystem::FRetryDomainsPtr RetryDomains = MakeShared<FHttpRetrySystem::FRetryDomains>(MoveTemp(AltDomains));
 	TSharedRef<IHttpRequest> HttpRequest = HttpRetryManager->CreateRequest(
@@ -2343,6 +2339,83 @@ TEST_CASE_METHOD(FWaitUntilCompleteHttpFixture, "Retry immediately without lock
 	HttpRequest->ProcessRequest();
 }
 
+TEST_CASE_METHOD(FWaitUntilCompleteHttpFixture, "Retry request with CompleteOnHttpThread policy won't add same duplicated request to http manager", HTTP_TAG)
+{
+	if (!bRetryEnabled)
+	{
+		return;
+	}
+
+	DisableWarningsInThisTest();
+
+	HttpModule->HttpConnectionTimeout = 1.0f;
+
+	FString OriginalUrl;
+
+	TArray<FString> AltDomains;
+	FString CorrectDomain = UrlDomain();
+
+	SECTION("With invalid domain")
+	{
+		SetInvalidDomainToTestConnectTimeout();
+		OriginalUrl = UrlToTestMethods();
+	}
+	SECTION("With mock connection failure for valid domain")
+	{
+		CVarHttpUrlPatternsToMockFailure->Set(TEXT("www.epicgames.com->0"));
+		OriginalUrl = UrlToTestMethods();
+		OriginalUrl = OriginalUrl.Replace(*UrlDomain(), TEXT("www.epicgames.com"));
+	}
+
+	AltDomains.Add(CorrectDomain);
+	FHttpRetrySystem::FRetryDomainsPtr RetryDomains = MakeShared<FHttpRetrySystem::FRetryDomains>(MoveTemp(AltDomains));
+	TSharedRef<IHttpRequest> HttpRequest = HttpRetryManager->CreateRequest(
+		1/*InRetryLimitCountOverride*/,
+		FHttpRetrySystem::FRetryTimeoutRelativeSecondsSetting()/*InRetryTimeoutRelativeSecondsOverride unused*/,
+		{ EHttpResponseCodes::TooManyRequests, EHttpResponseCodes::ServiceUnavail }/*InRetryResponseCodes*/,
+		FHttpRetrySystem::FRetryVerbs(),
+		RetryDomains
+	);
+
+	HttpRequest->SetURL(OriginalUrl);
+	HttpRequest->SetVerb(TEXT("GET"));
+	HttpRequest->SetDelegateThreadPolicy(EHttpRequestDelegateThreadPolicy::CompleteOnHttpThread);
+	++ExpectingExtraCallbacks;
+	HttpRequest->OnRequestWillRetry().BindLambda([this](FHttpRequestPtr /*Request*/, FHttpResponsePtr /*Response*/, float LockoutPeriod) {
+		--ExpectingExtraCallbacks;
+		CHECK(LockoutPeriod == 0);
+	});
+	HttpRequest->OnProcessRequestComplete().BindLambda([](FHttpRequestPtr HttpRequest, FHttpResponsePtr HttpResponse, bool bSucceeded) {
+		CHECK(bSucceeded);
+		REQUIRE(HttpResponse != nullptr);
+		CHECK(HttpResponse->GetResponseCode() == 200);
+	});
+	HttpRequest->ProcessRequest();
+
+	CVarHttpUrlPatternsToMockFailure->Set(TEXT(""));
+}
+
+TEST_CASE_METHOD(FWaitUntilQuitFromTestFixture, "Will keep request reference with mock connection failure which using CompleteOnHttpThread policy until mocked connection timeout", HTTP_TAG)
+{
+	DisableWarningsInThisTest();
+
+	CVarHttpUrlPatternsToMockFailure->Set(TEXT("www.epicgames.com->0"));
+	HttpModule->HttpConnectionTimeout = 1.0f;
+
+	TSharedRef<IHttpRequest> HttpRequest = CreateRequest();
+	HttpRequest->SetURL(TEXT("https://www.epicgames.com"));
+	HttpRequest->SetDelegateThreadPolicy(EHttpRequestDelegateThreadPolicy::CompleteOnHttpThread);
+	++ExpectingExtraCallbacks;
+	HttpRequest->OnProcessRequestComplete().BindLambda([this](FHttpRequestPtr HttpRequest, FHttpResponsePtr HttpResponse, bool bSucceeded) {
+		--ExpectingExtraCallbacks;
+		CHECK(!bSucceeded);
+		REQUIRE(HttpResponse == nullptr);
+		bQuitRequested = true;
+	});
+	HttpRequest->ProcessRequest();
+	CVarHttpUrlPatternsToMockFailure->Set(TEXT(""));
+}
+
 TEST_CASE_METHOD(FWaitUntilCompleteHttpFixture, "Retry will add default domain to alt domains", HTTP_TAG)
 {
 	if (!bRetryEnabled)
diff --git a/Engine/Source/Runtime/Online/HTTP/Private/GenericPlatform/HttpRequestCommon.cpp b/Engine/Source/Runtime/Online/HTTP/Private/GenericPlatform/HttpRequestCommon.cpp
--- a/Engine/Source/Runtime/Online/HTTP/Private/GenericPlatform/HttpRequestCommon.cpp
+++ b/Engine/Source/Runtime/Online/HTTP/Private/GenericPlatform/HttpRequestCommon.cpp
@@ -109,13 +109,11 @@ bool FHttpRequestCommon::TriggerMockFailure()
 		if (MockResponseCode.GetValue() == EHttpResponseCodes::Unknown)
 		{
 			int32 HttpConnectionTimeout = FHttpModule::Get().GetHttpConnectionTimeout();
-			TWeakPtr<FHttpRequestCommon> RequestWeakPtr(SharedThis(this));
-			FHttpModule::Get().GetHttpManager().AddHttpThreadTask([RequestWeakPtr]() {
-				if (TSharedPtr<FHttpRequestCommon> RequestPtr = RequestWeakPtr.Pin())
-				{
-					RequestPtr->SetFailureReason(EHttpFailureReason::ConnectionError);
-					RequestPtr->FinishRequestNotInHttpManager();
-				}
+			TSharedPtr<FHttpRequestCommon> RequestPtr(SharedThis(this));
+			FHttpModule::Get().GetHttpManager().AddHttpThreadTask([RequestPtr]() mutable {
+				RequestPtr->SetFailureReason(EHttpFailureReason::ConnectionError);
+				RequestPtr->FinishRequestNotInHttpManager();
+				RequestPtr.Reset();
 			}, HttpConnectionTimeout);
 
 			// Connect timeout mocking will trigger FinishRequest after a delay, still make sure total timeout 
diff --git a/Engine/Source/Runtime/Online/HTTP/Private/HttpManager.cpp b/Engine/Source/Runtime/Online/HTTP/Private/HttpManager.cpp
--- a/Engine/Source/Runtime/Online/HTTP/Private/HttpManager.cpp
+++ b/Engine/Source/Runtime/Online/HTTP/Private/HttpManager.cpp
@@ -46,6 +46,13 @@ TAutoConsoleVariable<FString> CVarHttpUrlPatternsToMockFailure(
 	ECVF_SaveForNextBoot
 );
 
+TAutoConsoleVariable<bool> CVarHttpRemoveRequestUsingHttpThreadPolicyOnHttpThread(
+	TEXT("http.RemoveRequestUsingHttpThreadPolicyOnHttpThread"),
+	true,
+	TEXT("When enabled, if http request use CompleteOnHttpThread, once complete it will be removed from http manager on http thread"),
+	ECVF_SaveForNextBoot
+);
+
 // FHttpManager
 
 const TCHAR* LexToString(const EHttpFlushReason& FlushReason)
@@ -521,6 +528,7 @@ bool FHttpManager::Tick(float DeltaSeconds)
 		{
 			FHttpRequestRef CompletedRequestRef = CompletedRequest->AsShared();
 
+			if (!CVarHttpRemoveRequestUsingHttpThreadPolicyOnHttpThread.GetValueOnAnyThread() || CompletedRequest->GetDelegateThreadPolicy() == EHttpRequestDelegateThreadPolicy::CompleteOnGameThread)
 			{
 				FScopeLock ScopeLock(&RequestLock);
 				Requests.Remove(CompletedRequestRef);
diff --git a/Engine/Source/Runtime/Online/HTTP/Private/HttpRetrySystem.cpp b/Engine/Source/Runtime/Online/HTTP/Private/HttpRetrySystem.cpp
--- a/Engine/Source/Runtime/Online/HTTP/Private/HttpRetrySystem.cpp
+++ b/Engine/Source/Runtime/Online/HTTP/Private/HttpRetrySystem.cpp
@@ -15,19 +15,6 @@
 #include "Misc/ConfigCacheIni.h"
 #include "Misc/CoreDelegates.h"
 
-TAutoConsoleVariable<bool> CVarHttpRetrySystemNonGameThreadSupportEnabled(
-	TEXT("Http.RetrySystemNonGameThreadSupportEnabled"),
-	true,
-	TEXT("Enable retry system non-game thread support")
-);
-
-TAutoConsoleVariable<bool> CVarHttpAddDomainOfDefaultUrlToRetryDomainsEnabled(
-	TEXT("http.OnlySetUrlFromAltDomainInRetry"),
-	true,
-	TEXT("When enabled, add domain in default url into retry domains"),
-	ECVF_SaveForNextBoot
-);
-
 LLM_DEFINE_TAG(HTTP);
 
 namespace FHttpRetrySystem
@@ -152,21 +139,18 @@ bool FHttpRetrySystem::FRequest::ProcessRequest()
 
 	OriginalUrl = HttpRequest->GetURL();
 
-	if (CVarHttpAddDomainOfDefaultUrlToRetryDomainsEnabled.GetValueOnAnyThread())
+	if (RetryDomains.IsValid() && !RetryDomains->Domains.IsEmpty())
 	{
-		if (RetryDomains.IsValid() && !RetryDomains->Domains.IsEmpty())
+		FString OriginalUrlDomainAndPort = FPlatformHttp::GetUrlDomainAndPort(OriginalUrl);
+		int32 Index = RetryDomains->Domains.Find(OriginalUrlDomainAndPort);
+		if (Index == INDEX_NONE)
 		{
-			FString OriginalUrlDomainAndPort = FPlatformHttp::GetUrlDomainAndPort(OriginalUrl);
-			int32 Index = RetryDomains->Domains.Find(OriginalUrlDomainAndPort);
-			if (Index == INDEX_NONE)
-			{
-				RetryDomains->Domains.Insert(MoveTemp(OriginalUrlDomainAndPort), 0);
-			}
-			else if (Index > 0)
-			{
-				RetryDomains->Domains.RemoveAt(Index);
-				RetryDomains->Domains.Insert(MoveTemp(OriginalUrlDomainAndPort), 0);
-			}
+			RetryDomains->Domains.Insert(MoveTemp(OriginalUrlDomainAndPort), 0);
+		}
+		else if (Index > 0)
+		{
+			RetryDomains->Domains.RemoveAt(Index);
+			RetryDomains->Domains.Insert(MoveTemp(OriginalUrlDomainAndPort), 0);
 		}
 	}
 
@@ -235,11 +219,6 @@ void FHttpRetrySystem::FRequest::HttpOnRequestProgress(FHttpRequestPtr InHttpReq
 
 void FHttpRetrySystem::FRequest::HttpOnProcessRequestComplete(FHttpRequestPtr InHttpRequest, FHttpResponsePtr HttpResponse, bool bSucceeded)
 {
-	if (!CVarHttpRetrySystemNonGameThreadSupportEnabled.GetValueOnAnyThread())
-	{
-		return;
-	}
-
 	TSharedPtr<FManager> RetryManagerPtr = RetryManager.Pin();
 	if (!RetryManagerPtr)
 	{
@@ -335,11 +314,6 @@ void FHttpRetrySystem::FRequest::HttpOnStatusCodeReceived(FHttpRequestPtr Reques
 
 void FHttpRetrySystem::FRequest::HttpOnHeaderReceived(FHttpRequestPtr Request, const FString& HeaderName, const FString& NewHeaderValue)
 {
-	if (!CVarHttpRetrySystemNonGameThreadSupportEnabled.GetValueOnAnyThread())
-	{
-		return;
-	}
-
 	TSharedRef<FRequest> SelfPtr = StaticCastSharedRef<FRequest>(AsShared());
 	OnHeaderReceived().ExecuteIfBound(SelfPtr, HeaderName, NewHeaderValue);
 }
@@ -642,17 +616,14 @@ bool FHttpRetrySystem::FManager::ProcessRequest(TSharedRef<FHttpRetrySystem::FRe
 {
 	QUICK_SCOPE_CYCLE_COUNTER(STAT_FHttpRetrySystem_FManager_ProcessRequest);
 
-	if (CVarHttpRetrySystemNonGameThreadSupportEnabled.GetValueOnAnyThread())
+	// Let the request trigger timeout by itself, instead of ticking it in retry system
+	if (HttpRetryRequest->RetryTimeoutRelativeSecondsOverride.IsSet())
 	{
-		// Let the request trigger timeout by itself, instead of ticking it in retry system
-		if (HttpRetryRequest->RetryTimeoutRelativeSecondsOverride.IsSet())
-		{
-			HttpRetryRequest->SetTimeout(HttpRetryRequest->RetryTimeoutRelativeSecondsOverride.GetValue());
-		}
-		else if (RetryTimeoutRelativeSecondsDefault.IsSet())
-		{
-			HttpRetryRequest->SetTimeout(RetryTimeoutRelativeSecondsDefault.GetValue());
-		}
+		HttpRetryRequest->SetTimeout(HttpRetryRequest->RetryTimeoutRelativeSecondsOverride.GetValue());
+	}
+	else if (RetryTimeoutRelativeSecondsDefault.IsSet())
+	{
+		HttpRetryRequest->SetTimeout(RetryTimeoutRelativeSecondsDefault.GetValue());
 	}
 
 	FScopeLock ScopeLock(&RequestListLock);
diff --git a/Engine/Source/Runtime/Online/HTTP/Private/HttpThread.cpp b/Engine/Source/Runtime/Online/HTTP/Private/HttpThread.cpp
--- a/Engine/Source/Runtime/Online/HTTP/Private/HttpThread.cpp
+++ b/Engine/Source/Runtime/Online/HTTP/Private/HttpThread.cpp
@@ -31,6 +31,8 @@ TAutoConsoleVariable<int32> CVarHttpDurationInQueueToWarnInSeconds(
 	TEXT("If http request waited more than this duration in the queue, output a warning before starting")
 );
 
+extern TAutoConsoleVariable<bool> CVarHttpRemoveRequestUsingHttpThreadPolicyOnHttpThread;
+
 // Thread priority cvar (settable at runtime)
 // We declare these explicitly rather than just casting the cvar in case the enum changes in future
 const int32 GHttpThreadPriorities[] =
@@ -376,16 +378,25 @@ void FHttpThreadBase::FinishRequestsFromHttpThreadWithCallbacks(TArray<FHttpRequ
 
 			CompleteThreadedRequest(Request);
 
+			FHttpRequestRef RequestRef = Request->AsShared();
+
 			if (Request->GetDelegateThreadPolicy() == EHttpRequestDelegateThreadPolicy::CompleteOnHttpThread)
 			{
 				FHttpManager& HttpManager = FHttpModule::Get().GetHttpManager();
+				if (CVarHttpRemoveRequestUsingHttpThreadPolicyOnHttpThread.GetValueOnAnyThread())
+				{
+					HttpManager.RemoveRequest(RequestRef);
+				}
 				HttpManager.MarkCurrentThreadCompletingRequest(true);
 				Request->FinishRequest();
-				HttpManager.BroadcastHttpRequestCompleted(Request->AsShared());
+				HttpManager.BroadcastHttpRequestCompleted(RequestRef);
 				HttpManager.MarkCurrentThreadCompletingRequest(false);
 			}
 
-			CompletedThreadedRequests.Enqueue(Request);
+			if (!CVarHttpRemoveRequestUsingHttpThreadPolicyOnHttpThread.GetValueOnAnyThread() || RequestRef->GetDelegateThreadPolicy() == EHttpRequestDelegateThreadPolicy::CompleteOnGameThread)
+			{
+				CompletedThreadedRequests.Enqueue(Request);
+			}
 		}
 		RequestsToComplete.Reset();
 	}
