Fix the compile issue in event loop tests encounter when using Clang 18.1.8;
Remove the redundant line of test code where InvalidDomain no longer needed for http test case.

diff --git a/Engine/Source/Programs/WebTests/Private/TestHttp.cpp b/Engine/Source/Programs/WebTests/Private/TestHttp.cpp
--- a/Engine/Source/Programs/WebTests/Private/TestHttp.cpp
+++ b/Engine/Source/Programs/WebTests/Private/TestHttp.cpp
@@ -2310,11 +2310,8 @@ TEST_CASE_METHOD(FWaitUntilCompleteHttpFixture, "Retry immediately without lock
 
 	TArray<FString> AltDomains;
 	FString CorrectDomain = UrlDomain();
-
-	SetInvalidDomainToTestConnectTimeout();
-
-	FString InvalidDomain = UrlDomain();
 	AltDomains.Add(CorrectDomain);
+
 	FHttpRetrySystem::FRetryDomainsPtr RetryDomains = MakeShared<FHttpRetrySystem::FRetryDomains>(MoveTemp(AltDomains));
 	TSharedRef<IHttpRequest> HttpRequest = HttpRetryManager->CreateRequest(
 		1/*InRetryLimitCountOverride*/,
@@ -2324,6 +2321,7 @@ TEST_CASE_METHOD(FWaitUntilCompleteHttpFixture, "Retry immediately without lock
 		RetryDomains
 	);
 
+	SetInvalidDomainToTestConnectTimeout();
 	HttpRequest->SetURL(UrlToTestMethods());
 	HttpRequest->SetVerb(TEXT("GET"));
 	++ExpectingExtraCallbacks;
diff --git a/Engine/Source/Runtime/Online/Experimental/EventLoopTests/Tests/EventLoop/EventLoopIOManagerBSDSocketSelectTests.cpp b/Engine/Source/Runtime/Online/Experimental/EventLoopTests/Tests/EventLoop/EventLoopIOManagerBSDSocketSelectTests.cpp
--- a/Engine/Source/Runtime/Online/Experimental/EventLoopTests/Tests/EventLoop/EventLoopIOManagerBSDSocketSelectTests.cpp
+++ b/Engine/Source/Runtime/Online/Experimental/EventLoopTests/Tests/EventLoop/EventLoopIOManagerBSDSocketSelectTests.cpp
@@ -848,4 +848,4 @@ namespace UE::EventLoop
 
 } // UE::EventLoop
 
-#endif PLATFORM_HAS_BSD_SOCKETS
+#endif // PLATFORM_HAS_BSD_SOCKETS
