Move IModuleInterface methods into public section for HTTP classes.

The IModuleInterface base class methods are public, so making the subclass methods private did not actually prevent access. It just made unit tests more awkward to write.

diff --git a/Engine/Source/Programs/WebTests/Private/TestHttp.cpp b/Engine/Source/Programs/WebTests/Private/TestHttp.cpp
--- a/Engine/Source/Programs/WebTests/Private/TestHttp.cpp
+++ b/Engine/Source/Programs/WebTests/Private/TestHttp.cpp
@@ -129,8 +129,7 @@ public:
 	void InitModule()
 	{
 		HttpModule = new FMockHttpModule();
-		IModuleInterface* Module = HttpModule;
-		Module->StartupModule();
+		HttpModule->StartupModule();
 		if (bRetryEnabled)
 		{
 			HttpRetryManager = MakeShared<FMockRetryManager>(FHttpRetrySystem::FRetryLimitCountSetting(0), FHttpRetrySystem::FRetryTimeoutRelativeSecondsSetting(/*RetryTimeoutRelativeSeconds*/));
@@ -141,11 +140,10 @@ public:
 	{
 		HttpRetryManager = nullptr;
 
-		IModuleInterface* Module = HttpModule;
-		if (Module)
+		if (HttpModule)
 		{
-			Module->ShutdownModule();
-			delete Module;
+			HttpModule->ShutdownModule();
+			delete HttpModule;
 			HttpModule = nullptr;
 		}
 	}
@@ -2783,8 +2781,7 @@ public:
 	FLocalHttpServerFixture()
 	{
 		HttpServerModule = new FHttpServerModule();
-		IModuleInterface* Module = HttpServerModule;
-		Module->StartupModule();
+		HttpServerModule->StartupModule();
 
 		HttpRouter = HttpServerModule->GetHttpRouter(LocalHttpServerPort);
 		CHECK(HttpRouter.IsValid());
@@ -2809,9 +2806,8 @@ public:
 		HttpRouter->UnbindRoute(HttpRouteHandle);
 		HttpRouter.Reset();
 
-		IModuleInterface* Module = HttpServerModule;
-		Module->ShutdownModule();
-		delete Module;
+		HttpServerModule->ShutdownModule();
+		delete HttpServerModule;
 	}
 
 	TSharedPtr<IHttpRouter> HttpRouter;
diff --git a/Engine/Source/Runtime/Online/HTTPServer/Public/HttpServerModule.h b/Engine/Source/Runtime/Online/HTTPServer/Public/HttpServerModule.h
--- a/Engine/Source/Runtime/Online/HTTPServer/Public/HttpServerModule.h
+++ b/Engine/Source/Runtime/Online/HTTPServer/Public/HttpServerModule.h
@@ -79,8 +79,6 @@ public:
 	  */
 	 HTTPSERVER_API bool HasPendingListeners() const;
 
-private:
-
 	// IModuleInterface
 
 	/**
