[Http]
Fix the failure on xCurl when use http2.

Reason: xCurl doesn't like Content-Length: 0 header when send the request without any content. Now it will only be set when content is not empty.

diff --git a/Engine/Source/Programs/WebTests/Private/TestHttp.cpp b/Engine/Source/Programs/WebTests/Private/TestHttp.cpp
--- a/Engine/Source/Programs/WebTests/Private/TestHttp.cpp
+++ b/Engine/Source/Programs/WebTests/Private/TestHttp.cpp
@@ -389,22 +389,76 @@ TEST_CASE_METHOD(FWaitUntilCompleteHttpFixture, "Http method use Get by default"
 	HttpRequest->ProcessRequest();
 }
 
-TEST_CASE_METHOD(FWaitUntilCompleteHttpFixture, "Http Methods", HTTP_TAG)
+TEST_CASE_METHOD(FWaitUntilCompleteHttpFixture, "Http methods without payloads", HTTP_TAG)
 {
 	TSharedRef<IHttpRequest> HttpRequest = CreateRequest();
-	CHECK(HttpRequest->GetVerb() == TEXT("GET"));
-
 	HttpRequest->SetURL(UrlToTestMethods());
 
 #if UE_HTTP_SUPPORT_VERB_CONNECT
-	const ANSICHAR* Method = GENERATE("GET", "POST", "DELETE", "PUT", "PATCH", "HEAD", "OPTIONS", "TRACE", "CONNECT");
+	const ANSICHAR* Method = GENERATE("GET", "DELETE", "HEAD", "OPTIONS", "TRACE", "CONNECT");
 #else // Can't use macro UE_HTTP_SUPPORT_VERB_CONNECT inside GENERATE macro, so define it here again
-	const ANSICHAR* Method = GENERATE("GET", "POST", "DELETE", "PUT", "PATCH", "HEAD", "OPTIONS", "TRACE");
+	const ANSICHAR* Method = GENERATE("GET", "DELETE", "HEAD", "OPTIONS", "TRACE");
 #endif
 
-	DYNAMIC_SECTION("Test methods with " << Method << " request")
+	const TCHAR* Version = GENERATE(FHttpConstants::VERSION_2TLS, FHttpConstants::VERSION_1_1);
+
+	DYNAMIC_SECTION("Test methods with " << Method << " request on version " << Version)
+	{
+		FString MethodStr = FString(Method);
+		HttpRequest->SetVerb(MethodStr);
+		HttpRequest->SetOption(HttpRequestOptions::HttpVersion, Version);
+	}
+
+	HttpRequest->OnProcessRequestComplete().BindLambda([](FHttpRequestPtr HttpRequest, FHttpResponsePtr HttpResponse, bool bSucceeded) {
+		CHECK(bSucceeded);
+		REQUIRE(HttpResponse != nullptr);
+		CHECK(HttpResponse->GetResponseCode() == 200);
+	});
+	HttpRequest->ProcessRequest();
+}
+
+TEST_CASE_METHOD(FWaitUntilCompleteHttpFixture, "Http methods with payloads", HTTP_TAG)
+{
+	TSharedRef<IHttpRequest> HttpRequest = CreateRequest();
+	HttpRequest->SetURL(UrlToTestMethods());
+
+	const ANSICHAR* Method = GENERATE("POST", "PUT", "PATCH");
+	const TCHAR* Version = GENERATE(FHttpConstants::VERSION_2TLS, FHttpConstants::VERSION_1_1);
+
+	DYNAMIC_SECTION("Test methods with " << Method << " request on version " << Version)
+	{
+		FString MethodStr = FString(Method);
+		HttpRequest->SetVerb(MethodStr);
+		HttpRequest->SetOption(HttpRequestOptions::HttpVersion, Version);
+	}
+
+	uint32 ChunkSize = 1024 * 1024;
+	TArray<uint8> DataChunk;
+	DataChunk.SetNum(ChunkSize);
+	FMemory::Memset(DataChunk.GetData(), 'd', ChunkSize);
+	HttpRequest->SetContent(MoveTemp(DataChunk));
+
+	HttpRequest->OnProcessRequestComplete().BindLambda([](FHttpRequestPtr HttpRequest, FHttpResponsePtr HttpResponse, bool bSucceeded) {
+		CHECK(bSucceeded);
+		REQUIRE(HttpResponse != nullptr);
+		CHECK(HttpResponse->GetResponseCode() == 200);
+	});
+	HttpRequest->ProcessRequest();
+}
+
+TEST_CASE_METHOD(FWaitUntilCompleteHttpFixture, "Http methods which support to upload payload can work well without any payload", HTTP_TAG)
+{
+	TSharedRef<IHttpRequest> HttpRequest = CreateRequest();
+	HttpRequest->SetURL(UrlToTestMethods());
+
+	const ANSICHAR* Method = GENERATE("POST", "PUT", "PATCH");
+	const TCHAR* Version = GENERATE(FHttpConstants::VERSION_2TLS, FHttpConstants::VERSION_1_1);
+
+	DYNAMIC_SECTION("Test methods with " << Method << " request on version " << Version)
 	{
-		HttpRequest->SetVerb(FString(Method));
+		FString MethodStr = FString(Method);
+		HttpRequest->SetVerb(MethodStr);
+		HttpRequest->SetOption(HttpRequestOptions::HttpVersion, Version);
 	}
 
 	HttpRequest->OnProcessRequestComplete().BindLambda([](FHttpRequestPtr HttpRequest, FHttpResponsePtr HttpResponse, bool bSucceeded) {
@@ -2990,15 +3044,3 @@ TEST_CASE_METHOD(FWaitUntilCompleteHttpFixture, "Http Methods over Unix Domain S
 
 #endif //UE_HTTP_SUPPORT_UNIX_SOCKET
 
-TEST_CASE_METHOD(FWaitUntilCompleteHttpFixture, "Http Methods over http2", HTTP_TAG)
-{
-	TSharedRef<IHttpRequest> HttpRequest = CreateRequest();
-	HttpRequest->SetURL(UrlToTestMethods());
-	HttpRequest->SetOption(HttpRequestOptions::HttpVersion, FHttpConstants::VERSION_2TLS);
-	HttpRequest->OnProcessRequestComplete().BindLambda([](FHttpRequestPtr HttpRequest, FHttpResponsePtr HttpResponse, bool bSucceeded) {
-		CHECK(bSucceeded);
-		REQUIRE(HttpResponse != nullptr);
-		CHECK(HttpResponse->GetResponseCode() == 200);
-	});
-	HttpRequest->ProcessRequest();
-}
diff --git a/Engine/Source/Runtime/Online/HTTP/Private/Curl/CurlHttp.cpp b/Engine/Source/Runtime/Online/HTTP/Private/Curl/CurlHttp.cpp
--- a/Engine/Source/Runtime/Online/HTTP/Private/Curl/CurlHttp.cpp
+++ b/Engine/Source/Runtime/Online/HTTP/Private/Curl/CurlHttp.cpp
@@ -27,6 +27,13 @@ TAutoConsoleVariable<bool> CVarHttpCurlAllowHTTP2(
 	ECVF_Default
 );
 
+TAutoConsoleVariable<bool> CVarHttpAlwaysSetContentLengthEnabled(
+	TEXT("http.AlwaysSetContentLengthEnabled"),
+	false,
+	TEXT("Always set Content-Length header even if there is no content to upload when do the http request"),
+	ECVF_SaveForNextBoot
+);
+
 #if WITH_SSL
 static int SslCertVerify(int PreverifyOk, X509_STORE_CTX* Context)
 {
@@ -821,7 +828,10 @@ bool FCurlHttpRequest::SetupRequest()
 	// content-length should be present http://www.w3.org/Protocols/rfc2616/rfc2616-sec4.html#sec4.4
 	if (GetHeader(TEXT("Content-Length")).IsEmpty())
 	{
-		SetHeader(TEXT("Content-Length"), FString::Printf(TEXT("%llu"), RequestPayload->GetContentLength()));
+		if (RequestPayload->GetContentLength() != 0 || CVarHttpAlwaysSetContentLengthEnabled.GetValueOnAnyThread())
+		{
+			SetHeader(TEXT("Content-Length"), FString::Printf(TEXT("%llu"), RequestPayload->GetContentLength()));
+		}
 	}
 
 	// Remove "Expect: 100-continue" since this is supposed to cause problematic behavior on Amazon ELB (and WinInet doesn't send that either)
