Add priority to http request, so when concurrent http requests reach max limit, new request can be queued according to priority.

diff --git a/Engine/Plugins/Media/ElectraUtil/Source/ElectraHTTPStream/Private/Generic/GenericElectraHTTPStream.cpp b/Engine/Plugins/Media/ElectraUtil/Source/ElectraHTTPStream/Private/Generic/GenericElectraHTTPStream.cpp
--- a/Engine/Plugins/Media/ElectraUtil/Source/ElectraHTTPStream/Private/Generic/GenericElectraHTTPStream.cpp
+++ b/Engine/Plugins/Media/ElectraUtil/Source/ElectraHTTPStream/Private/Generic/GenericElectraHTTPStream.cpp
@@ -320,6 +320,7 @@ bool FElectraHTTPStreamRequestGeneric::Setup()
 	RequestHandle->SetVerb(Verb);
 	RequestHandle->SetURL(URL);
 	RequestHandle->SetDelegateThreadPolicy(EHttpRequestDelegateThreadPolicy::CompleteOnHttpThread);
+	RequestHandle->SetPriority(EHttpRequestPriority::High);
 	RequestHandle->OnProcessRequestComplete().BindThreadSafeSP(AsShared(), &FElectraHTTPStreamRequestGeneric::OnProcessRequestComplete);
 	RequestHandle->OnHeaderReceived().BindThreadSafeSP(AsShared(), &FElectraHTTPStreamRequestGeneric::OnHeaderReceived);
 	RequestHandle->OnStatusCodeReceived().BindThreadSafeSP(AsShared(), &FElectraHTTPStreamRequestGeneric::OnStatusCodeReceived);
diff --git a/Engine/Source/Programs/WebTests/Private/TestHttp.cpp b/Engine/Source/Programs/WebTests/Private/TestHttp.cpp
--- a/Engine/Source/Programs/WebTests/Private/TestHttp.cpp
+++ b/Engine/Source/Programs/WebTests/Private/TestHttp.cpp
@@ -1610,7 +1610,6 @@ TEST_CASE_METHOD(FWaitUntilCompleteHttpFixture, "Redirect enabled by default and
 	HttpRequest->ProcessRequest();
 }
 
-
 class FWaitUntilQuitFromTestFixture : public FWaitUntilCompleteHttpFixture
 {
 public:
@@ -2219,6 +2218,78 @@ TEST_CASE_METHOD(FWaitThreadedHttpFixture, "Cancel http request with ProcessRequ
 	ThreadedHttpRunnable.StartTestHttpThread(true/*bBlockGameThread*/);
 }
 
+TEST_CASE_METHOD(FWaitThreadedHttpFixture, "Requests with different priorities will be inserted to correct position in the queue", HTTP_TAG)
+{
+	CVarHttpMaxConcurrentRequests->Set(1);
+
+	std::atomic<bool> bRunningRequestCompleted = false;
+	std::atomic<bool> bHttpRequestWithNormalPriorityCompleted = false;
+	std::atomic<bool> bHttpRequestWithHighPriorityCompleted = false;
+	std::atomic<bool> bHttpRequestWithLowPriorityCompleted = false;
+
+	// Use thread here so we can check local bool flags in the complete delegates
+	ThreadedHttpRunnable.OnRunFromThread().BindLambda([this, &bRunningRequestCompleted, &bHttpRequestWithNormalPriorityCompleted, &bHttpRequestWithHighPriorityCompleted, &bHttpRequestWithLowPriorityCompleted]() {
+		TSharedRef<IHttpRequest> HttpRequestRunning = CreateRequest();
+		HttpRequestRunning->SetURL(UrlStreamDownload(3/*Chunks*/, HTTP_TEST_TIMEOUT_CHUNK_SIZE, 1/*ChunkLatency*/));
+		HttpRequestRunning->SetDelegateThreadPolicy(EHttpRequestDelegateThreadPolicy::CompleteOnHttpThread);
+		HttpRequestRunning->OnProcessRequestComplete().BindLambda([this, &bRunningRequestCompleted, &bHttpRequestWithNormalPriorityCompleted, &bHttpRequestWithHighPriorityCompleted, &bHttpRequestWithLowPriorityCompleted](FHttpRequestPtr HttpRequestQueuing, FHttpResponsePtr HttpResponse, bool bSucceeded) {
+			CHECK(bSucceeded);
+			bRunningRequestCompleted = true;
+			CHECK(!bHttpRequestWithNormalPriorityCompleted);
+			CHECK(!bHttpRequestWithHighPriorityCompleted);
+			CHECK(!bHttpRequestWithLowPriorityCompleted);
+		});
+		HttpRequestRunning->ProcessRequest();
+
+		TSharedRef<IHttpRequest> HttpRequestWithNormalPriority = CreateRequest();
+		HttpRequestWithNormalPriority->SetURL(UrlToTestMethods());
+		HttpRequestWithNormalPriority->SetDelegateThreadPolicy(EHttpRequestDelegateThreadPolicy::CompleteOnHttpThread);
+		HttpRequestWithNormalPriority->OnProcessRequestComplete().BindLambda([this, &bRunningRequestCompleted, &bHttpRequestWithNormalPriorityCompleted, &bHttpRequestWithHighPriorityCompleted, &bHttpRequestWithLowPriorityCompleted](FHttpRequestPtr HttpRequestNormal, FHttpResponsePtr HttpResponse, bool bSucceeded) {
+			CHECK(bSucceeded);
+			CHECK(bRunningRequestCompleted);
+			bHttpRequestWithNormalPriorityCompleted = true;
+			CHECK(bHttpRequestWithHighPriorityCompleted);
+			CHECK(!bHttpRequestWithLowPriorityCompleted);
+		});
+		HttpRequestWithNormalPriority->ProcessRequest();
+
+		TSharedRef<IHttpRequest> HttpRequestWithHighPriority = CreateRequest();
+		HttpRequestWithHighPriority->SetURL(UrlToTestMethods());
+		HttpRequestWithHighPriority->SetPriority(EHttpRequestPriority::High);
+		HttpRequestWithHighPriority->SetDelegateThreadPolicy(EHttpRequestDelegateThreadPolicy::CompleteOnHttpThread);
+		HttpRequestWithHighPriority->OnProcessRequestComplete().BindLambda([this, &bRunningRequestCompleted, &bHttpRequestWithNormalPriorityCompleted, &bHttpRequestWithHighPriorityCompleted, &bHttpRequestWithLowPriorityCompleted](FHttpRequestPtr HttpRequestNormal, FHttpResponsePtr HttpResponse, bool bSucceeded) {
+			CHECK(bSucceeded);
+			CHECK(bRunningRequestCompleted);
+			CHECK(!bHttpRequestWithNormalPriorityCompleted);
+			bHttpRequestWithHighPriorityCompleted = true;
+			CHECK(!bHttpRequestWithLowPriorityCompleted);
+		});
+		HttpRequestWithHighPriority->ProcessRequest();
+
+		TSharedRef<IHttpRequest> HttpRequestWithLowPriority = CreateRequest();
+		HttpRequestWithLowPriority->SetURL(UrlToTestMethods());
+		HttpRequestWithLowPriority->SetPriority(EHttpRequestPriority::Low);
+		HttpRequestWithLowPriority->SetDelegateThreadPolicy(EHttpRequestDelegateThreadPolicy::CompleteOnHttpThread);
+		HttpRequestWithLowPriority->OnProcessRequestComplete().BindLambda([this, &bRunningRequestCompleted, &bHttpRequestWithNormalPriorityCompleted, &bHttpRequestWithHighPriorityCompleted, &bHttpRequestWithLowPriorityCompleted](FHttpRequestPtr HttpRequestNormal, FHttpResponsePtr HttpResponse, bool bSucceeded) {
+			CHECK(bSucceeded);
+			CHECK(bRunningRequestCompleted);
+			CHECK(bHttpRequestWithNormalPriorityCompleted);
+			CHECK(bHttpRequestWithHighPriorityCompleted);
+			bHttpRequestWithLowPriorityCompleted = true;
+
+			ThreadedHttpRunnable.UnblockGameThread();
+		});
+		HttpRequestWithLowPriority->ProcessRequest();
+
+		FPlatformProcess::Sleep(1); // Make sure the first request started
+
+		FHttpStats HttpStats = HttpModule->GetHttpManager().GetHttpStats();
+		CHECK(HttpStats.RequestsInQueue == 3);
+	});
+
+	ThreadedHttpRunnable.StartTestHttpThread(true/*bBlockGameThread*/);
+}
+
 TEST_CASE_METHOD(FWaitUntilCompleteHttpFixture, "Cancel http request connect before timeout", HTTP_TAG)
 {
 	DisableWarningsInThisTest();
diff --git a/Engine/Source/Runtime/Online/BuildPatchServices/Private/Tests/Mock/HttpRequest.mock.h b/Engine/Source/Runtime/Online/BuildPatchServices/Private/Tests/Mock/HttpRequest.mock.h
--- a/Engine/Source/Runtime/Online/BuildPatchServices/Private/Tests/Mock/HttpRequest.mock.h
+++ b/Engine/Source/Runtime/Online/BuildPatchServices/Private/Tests/Mock/HttpRequest.mock.h
@@ -236,6 +236,17 @@ namespace BuildPatchServices
 			return EHttpRequestDelegateThreadPolicy::CompleteOnGameThread;
 		}
 
+		virtual void SetPriority(EHttpRequestPriority InPriority) override
+		{
+			MOCK_FUNC_NOT_IMPLEMENTED("FMockHttpRequest::SetPriority");
+		}
+
+		virtual EHttpRequestPriority GetPriority() const override
+		{
+			MOCK_FUNC_NOT_IMPLEMENTED("FMockHttpRequest::GetPriority");
+			return EHttpRequestPriority::Normal;
+		}
+
 		virtual const FString& GetEffectiveURL() const override
 		{
 			MOCK_FUNC_NOT_IMPLEMENTED("FMockHttpRequest::GetEffectiveURL");
diff --git a/Engine/Source/Runtime/Online/HTTP/Private/GenericPlatform/GenericPlatformHttp.cpp b/Engine/Source/Runtime/Online/HTTP/Private/GenericPlatform/GenericPlatformHttp.cpp
--- a/Engine/Source/Runtime/Online/HTTP/Private/GenericPlatform/GenericPlatformHttp.cpp
+++ b/Engine/Source/Runtime/Online/HTTP/Private/GenericPlatform/GenericPlatformHttp.cpp
@@ -127,6 +127,8 @@ public:
 	virtual float GetElapsedTime() const override { return 0.0f; }
 	virtual void SetDelegateThreadPolicy(EHttpRequestDelegateThreadPolicy InDelegateThreadPolicy) override {}
 	virtual EHttpRequestDelegateThreadPolicy GetDelegateThreadPolicy() const override { return EHttpRequestDelegateThreadPolicy::CompleteOnGameThread; }
+	HTTP_API virtual void SetPriority(EHttpRequestPriority InPriority) override {}
+	HTTP_API virtual EHttpRequestPriority GetPriority() const override { return EHttpRequestPriority::Normal; }
 	virtual void SetTimeout(float InTimeoutSecs) override {}
 	virtual void ClearTimeout() override {}
 	virtual void ResetTimeoutStatus() override {}
diff --git a/Engine/Source/Runtime/Online/HTTP/Private/GenericPlatform/HttpRequestCommon.cpp b/Engine/Source/Runtime/Online/HTTP/Private/GenericPlatform/HttpRequestCommon.cpp
--- a/Engine/Source/Runtime/Online/HTTP/Private/GenericPlatform/HttpRequestCommon.cpp
+++ b/Engine/Source/Runtime/Online/HTTP/Private/GenericPlatform/HttpRequestCommon.cpp
@@ -669,6 +669,16 @@ const FString& FHttpRequestCommon::GetURL() const
 	return URL;
 }
 
+void FHttpRequestCommon::SetPriority(EHttpRequestPriority InPriority)
+{
+	Priority = InPriority;
+}
+
+EHttpRequestPriority FHttpRequestCommon::GetPriority() const
+{
+	return Priority;
+}
+
 bool FHttpRequestCommon::PassReceivedDataToStream(void* Ptr, int64 Length)
 {
 	const FScopeLock StreamLock(&ResponseBodyReceiveStreamCriticalSection);
diff --git a/Engine/Source/Runtime/Online/HTTP/Private/GenericPlatform/HttpRequestCommon.h b/Engine/Source/Runtime/Online/HTTP/Private/GenericPlatform/HttpRequestCommon.h
--- a/Engine/Source/Runtime/Online/HTTP/Private/GenericPlatform/HttpRequestCommon.h
+++ b/Engine/Source/Runtime/Online/HTTP/Private/GenericPlatform/HttpRequestCommon.h
@@ -61,6 +61,9 @@ public:
 	HTTP_API virtual void SetURL(const FString& InURL) override;
 	HTTP_API virtual const FString& GetURL() const override;
 
+	HTTP_API virtual void SetPriority(EHttpRequestPriority InPriority) override;
+	HTTP_API virtual EHttpRequestPriority GetPriority() const override;
+
 protected:
 	/**
 	 * Check if this request is valid or allowed, before actually process the request
@@ -196,4 +199,6 @@ protected:
 
 	/** Record the time started to wait in the queue */
 	float TimeStartedWaitingInQueue = 0.0f;
+
+	EHttpRequestPriority Priority = EHttpRequestPriority::Normal;
 };
diff --git a/Engine/Source/Runtime/Online/HTTP/Private/HttpRequestAdapter.cpp b/Engine/Source/Runtime/Online/HTTP/Private/HttpRequestAdapter.cpp
--- a/Engine/Source/Runtime/Online/HTTP/Private/HttpRequestAdapter.cpp
+++ b/Engine/Source/Runtime/Online/HTTP/Private/HttpRequestAdapter.cpp
@@ -179,3 +179,12 @@ EHttpRequestDelegateThreadPolicy FHttpRequestAdapterBase::GetDelegateThreadPolic
 	return HttpRequest->GetDelegateThreadPolicy(); 
 }
 
+void FHttpRequestAdapterBase::SetPriority(EHttpRequestPriority InPriority) 
+{ 
+	HttpRequest->SetPriority(InPriority); 
+}
+
+EHttpRequestPriority FHttpRequestAdapterBase::GetPriority() const 
+{ 
+	return HttpRequest->GetPriority(); 
+}
diff --git a/Engine/Source/Runtime/Online/HTTP/Private/HttpThread.cpp b/Engine/Source/Runtime/Online/HTTP/Private/HttpThread.cpp
--- a/Engine/Source/Runtime/Online/HTTP/Private/HttpThread.cpp
+++ b/Engine/Source/Runtime/Online/HTTP/Private/HttpThread.cpp
@@ -31,6 +31,13 @@ TAutoConsoleVariable<int32> CVarHttpDurationInQueueToWarnInSeconds(
 	TEXT("If http request waited more than this duration in the queue, output a warning before starting")
 );
 
+TAutoConsoleVariable<bool> CVarHttpRequestPriorityEnabled(
+	TEXT("http.RequestPriorityEnabled"),
+	true,
+	TEXT("Whether to enable http request priority in the queue when max requests limit reached."),
+	ECVF_SaveForNextBoot
+);
+
 extern TAutoConsoleVariable<bool> CVarHttpRemoveRequestUsingHttpThreadPolicyOnHttpThread;
 
 // Thread priority cvar (settable at runtime)
@@ -53,7 +60,6 @@ const TCHAR* GHttpThreadPriortyNames[] =
 	TEXT("TPri_AboveNormal")
 };
 
-// Warning: Due to a bug with http module console variables, this cvar is not settable via the console (or via -execcmds). It needs to be set via ini (or via -dpcvars on startup). Hotfixing is supported
 static int32 GHttpThreadPriorityIndex = 3; // EThreadPriority::TPri_Normal
 FAutoConsoleVariableRef CVarHttpThreadPriority(
 	TEXT("http.ThreadPriority"), 
@@ -249,7 +255,8 @@ void FHttpThreadBase::ConsumeCanceledRequestsAndNewRequests(TArray<FHttpRequestC
 		while (NewThreadedRequests.Dequeue(Request))
 		{
 			Request->StartWaitingInQueue();
-			RateLimitedThreadedRequests.Add(Request);
+
+			InsertToRateLimitedRequestsAccordingToPriority(Request);
 		}
 	}
 
@@ -272,6 +279,26 @@ void FHttpThreadBase::ConsumeCanceledRequestsAndNewRequests(TArray<FHttpRequestC
 	}
 }
 
+void FHttpThreadBase::InsertToRateLimitedRequestsAccordingToPriority(FHttpRequestCommon* Request)
+{
+	QUICK_SCOPE_CYCLE_COUNTER(STAT_FHttpThreadBase_InsertToRateLimitedRequestsAccordingToPriority);
+
+	if (CVarHttpRequestPriorityEnabled.GetValueOnAnyThread())
+	{
+		for (int32 i = 0; i < RateLimitedThreadedRequests.Num(); ++i)
+		{
+			FHttpRequestCommon* ExistingRequest = RateLimitedThreadedRequests[i];
+			if (ExistingRequest->GetPriority() < Request->GetPriority())
+			{
+				RateLimitedThreadedRequests.Insert(Request, i);
+				return;
+			}
+		}
+	}
+
+	RateLimitedThreadedRequests.Add(Request);
+}
+
 void FHttpThreadBase::StartRequestsWaitingInQueue(TArray<FHttpRequestCommon*>& RequestsToComplete)
 {
 	FHttpManager& HttpManager = FHttpModule::Get().GetHttpManager();
diff --git a/Engine/Source/Runtime/Online/HTTP/Private/HttpThread.h b/Engine/Source/Runtime/Online/HTTP/Private/HttpThread.h
--- a/Engine/Source/Runtime/Online/HTTP/Private/HttpThread.h
+++ b/Engine/Source/Runtime/Online/HTTP/Private/HttpThread.h
@@ -140,6 +140,7 @@ protected:
 
 private:
 	void ConsumeCanceledRequestsAndNewRequests(TArray<FHttpRequestCommon*>& RequestsToCancel, TArray<FHttpRequestCommon*>& RequestsToComplete);
+	void InsertToRateLimitedRequestsAccordingToPriority(FHttpRequestCommon* Request);
 	void MoveCompletingRequestsToCompletedRequests(TArray<FHttpRequestCommon*>& RequestsToComplete);
 	void StartRequestsWaitingInQueue(TArray<FHttpRequestCommon*>& RequestsToComplete);
 	void FinishRequestsFromHttpThreadWithCallbacks(TArray<FHttpRequestCommon*>& RequestsToComplete);
diff --git a/Engine/Source/Runtime/Online/HTTP/Private/NullHttp.cpp b/Engine/Source/Runtime/Online/HTTP/Private/NullHttp.cpp
--- a/Engine/Source/Runtime/Online/HTTP/Private/NullHttp.cpp
+++ b/Engine/Source/Runtime/Online/HTTP/Private/NullHttp.cpp
@@ -224,6 +224,15 @@ EHttpRequestDelegateThreadPolicy FNullHttpRequest::GetDelegateThreadPolicy() con
 	return EHttpRequestDelegateThreadPolicy::CompleteOnGameThread;
 }
 
+void FNullHttpRequest::SetPriority(EHttpRequestPriority InPriority)
+{
+}
+
+EHttpRequestPriority FNullHttpRequest::GetPriority() const
+{
+	return EHttpRequestPriority::Normal;
+}
+
 void FNullHttpRequest::SetTimeout(float InTimeoutSecs) 
 {
 	TimeoutSecs = InTimeoutSecs;
diff --git a/Engine/Source/Runtime/Online/HTTP/Private/NullHttp.h b/Engine/Source/Runtime/Online/HTTP/Private/NullHttp.h
--- a/Engine/Source/Runtime/Online/HTTP/Private/NullHttp.h
+++ b/Engine/Source/Runtime/Online/HTTP/Private/NullHttp.h
@@ -45,6 +45,8 @@ public:
 	virtual float GetElapsedTime() const override;
 	virtual void SetDelegateThreadPolicy(EHttpRequestDelegateThreadPolicy InThreadPolicy) override;
 	virtual EHttpRequestDelegateThreadPolicy GetDelegateThreadPolicy() const override;
+	virtual void SetPriority(EHttpRequestPriority InPriority) override;
+	virtual EHttpRequestPriority GetPriority() const override;
 	virtual void SetTimeout(float InTimeoutSecs) override;
 	virtual void ClearTimeout() override;
 	virtual void ResetTimeoutStatus() override;
diff --git a/Engine/Source/Runtime/Online/HTTP/Private/TransactionallySafeHttpRequest.cpp b/Engine/Source/Runtime/Online/HTTP/Private/TransactionallySafeHttpRequest.cpp
--- a/Engine/Source/Runtime/Online/HTTP/Private/TransactionallySafeHttpRequest.cpp
+++ b/Engine/Source/Runtime/Online/HTTP/Private/TransactionallySafeHttpRequest.cpp
@@ -23,6 +23,7 @@ public:
 			TSharedPtr<IHttpRequest> Req{FPlatformHttp::ConstructRequest()};
 			
 			Req->SetDelegateThreadPolicy(ThreadPolicy);
+			Req->SetPriority(Priority);
 			Req->SetVerb(MoveTemp(Verb));
 
 			if (!Url.IsEmpty())
@@ -258,6 +259,16 @@ public:
 		return ThreadPolicy;
 	}
 
+	void SetPriority(EHttpRequestPriority InPriority) override
+	{
+		Priority = InPriority;
+	}
+
+	EHttpRequestPriority GetPriority() const override
+	{
+		return Priority;
+	}
+
 	bool ProcessRequest() override
 	{
 		bProcessRequest = true;
@@ -339,6 +350,7 @@ private:
 	TMap<FString, FString> Headers;
 	TSharedPtr<FArchive> ResponseBodyReceiveStream;
 	EHttpRequestDelegateThreadPolicy ThreadPolicy = EHttpRequestDelegateThreadPolicy::CompleteOnGameThread;
+	EHttpRequestPriority Priority = EHttpRequestPriority::Normal;
 	TOptional<float> TimeoutSecs;
 	TOptional<float> ActivityTimeoutSecs;
 	TOptional<FHttpRequestCompleteDelegate> CompleteDelegate;
@@ -516,6 +528,16 @@ EHttpRequestDelegateThreadPolicy FTransactionallySafeHttpRequest::GetDelegateThr
 	return InnerRequest->GetDelegateThreadPolicy();
 }
 
+void FTransactionallySafeHttpRequest::SetPriority(EHttpRequestPriority InPriority)
+{
+	InnerRequest->SetPriority(InPriority);
+}
+
+EHttpRequestPriority FTransactionallySafeHttpRequest::GetPriority() const
+{
+	return InnerRequest->GetPriority();
+}
+
 void FTransactionallySafeHttpRequest::SetTimeout(float InTimeoutSecs)
 {
 	InnerRequest->SetTimeout(InTimeoutSecs);
diff --git a/Engine/Source/Runtime/Online/HTTP/Private/TransactionallySafeHttpRequest.h b/Engine/Source/Runtime/Online/HTTP/Private/TransactionallySafeHttpRequest.h
--- a/Engine/Source/Runtime/Online/HTTP/Private/TransactionallySafeHttpRequest.h
+++ b/Engine/Source/Runtime/Online/HTTP/Private/TransactionallySafeHttpRequest.h
@@ -52,6 +52,8 @@ public:
 	float GetElapsedTime() const override;
 	void SetDelegateThreadPolicy(EHttpRequestDelegateThreadPolicy InThreadPolicy) override;
 	EHttpRequestDelegateThreadPolicy GetDelegateThreadPolicy() const override;
+	virtual void SetPriority(EHttpRequestPriority InPriority) override;
+	virtual EHttpRequestPriority GetPriority() const override;
 	void SetTimeout(float InTimeoutSecs) override;
 	void ClearTimeout() override;
 	void ResetTimeoutStatus() override;
diff --git a/Engine/Source/Runtime/Online/HTTP/Public/HttpRequestAdapter.h b/Engine/Source/Runtime/Online/HTTP/Public/HttpRequestAdapter.h
--- a/Engine/Source/Runtime/Online/HTTP/Public/HttpRequestAdapter.h
+++ b/Engine/Source/Runtime/Online/HTTP/Public/HttpRequestAdapter.h
@@ -51,6 +51,8 @@ public:
 	HTTP_API virtual void Tick(float DeltaSeconds) override;
 	HTTP_API virtual void SetDelegateThreadPolicy(EHttpRequestDelegateThreadPolicy InThreadPolicy) override;
 	HTTP_API virtual EHttpRequestDelegateThreadPolicy GetDelegateThreadPolicy() const override;
+	HTTP_API virtual void SetPriority(EHttpRequestPriority InPriority) override;
+	HTTP_API virtual EHttpRequestPriority GetPriority() const override;
 
 protected:
 	TSharedRef<IHttpRequest> HttpRequest;
diff --git a/Engine/Source/Runtime/Online/HTTP/Public/Interfaces/IHttpRequest.h b/Engine/Source/Runtime/Online/HTTP/Public/Interfaces/IHttpRequest.h
--- a/Engine/Source/Runtime/Online/HTTP/Public/Interfaces/IHttpRequest.h
+++ b/Engine/Source/Runtime/Online/HTTP/Public/Interfaces/IHttpRequest.h
@@ -17,6 +17,19 @@ enum class EHttpRequestDelegateThreadPolicy : uint8
 	CompleteOnHttpThread,
 };
 
+/**
+ * Enumerates priority to decide how the request will be queued when max concurrent limit reached
+ * The priority will be compared through value, so it's important to keep the order from low to high
+ */
+enum class EHttpRequestPriority : uint8
+{
+	Lowest = 0,
+	Low,
+	Normal,
+	High,
+	Highest
+};
+
 class IHttpRequest;
 class IHttpResponse;
 
@@ -424,6 +437,20 @@ public:
 	 */
 	virtual void ProcessRequestUntilComplete() = 0;
 
+	/**
+	 * Set priority of the request, which will be used to decide the position in the queue when max concurrent limit reached
+	 *
+	 * @param InHttpRequestPriority - The priority to set
+	 */
+	virtual void SetPriority(EHttpRequestPriority InPriority) = 0;
+
+	/**
+	 * Get request priority, which is used to decide the position in the queue when max concurrent limit reached
+	 *
+	 * @return The request priority
+	 */
+	virtual EHttpRequestPriority GetPriority() const = 0;
+
 	/**
 	 * Destructor for overrides
 	 */
