UE: TestHttpDecode: 2GB is no longer a valid size for TArray (it has to be smaller).

diff --git a/Engine/Source/Programs/WebTests/Private/TestHttpDecode.cpp b/Engine/Source/Programs/WebTests/Private/TestHttpDecode.cpp
--- a/Engine/Source/Programs/WebTests/Private/TestHttpDecode.cpp
+++ b/Engine/Source/Programs/WebTests/Private/TestHttpDecode.cpp
@@ -10,17 +10,19 @@
 #define HTTP_SERVER_TEST_CASE(x, ...) TEST_CASE(x, HTTP_SERVER_SUITE_TAGS)
 
 // Helper to generate a 2GB Test String
-FString Generate2GBURL() {
-	FString twoGBURL;
-	int64 size{ MAX_int32 };
-	twoGBURL.Reserve(size);
+FString Generate2GBURL()
+{
+	FString TwoGBURL;
+	constexpr int32 Size{ MAX_int32 - 1 };
+	TwoGBURL.Reserve(Size);
 
-	for (int64 i = 0; i < (size - 2); ++i) {
-		twoGBURL.AppendChar('A');
+	for (int32 i = 0; i < (Size - 2); ++i)
+	{
+		TwoGBURL.AppendChar('A');
 	}
 
-	twoGBURL.AppendChar('%');
-	return twoGBURL;
+	TwoGBURL.AppendChar('%');
+	return TwoGBURL;
 }
 
 // To verify UE-172890
