Use thread pool to notify http thread to wake up from event loop notifier.

#Synced-CL 48203440

diff --git a/Engine/Source/Programs/WebTests/Private/TestHttp.cpp b/Engine/Source/Programs/WebTests/Private/TestHttp.cpp
--- a/Engine/Source/Programs/WebTests/Private/TestHttp.cpp
+++ b/Engine/Source/Programs/WebTests/Private/TestHttp.cpp
@@ -16,6 +16,7 @@
 #include "HttpServerModule.h"
 #include "Logging/LogScopedVerbosityOverride.h"
 #include "Misc/CommandLine.h"
+#include "Misc/QueuedThreadPool.h"
 #include "Misc/ScopeExit.h"
 #include "TestHarness.h"
 #include "TestMacros/Assertions.h"
@@ -207,7 +208,7 @@ public:
 
 TEST_CASE_METHOD(FHttpModuleTestFixture, "Shutdown http module without issue when there are ongoing upload http requests.", HTTP_TAG)
 {
-	DisableWarningsInThisTest();
+	DisableWarningsInThisTest(); // for: [FHttpManager::Shutdown] Unbinding delegates for 5 outstanding Http Requests
 
 	uint32 ChunkSize = 1024 * 1024;
 	TArray<uint8> DataChunk;
@@ -1148,6 +1149,55 @@ TEST_CASE_METHOD(FWaitUntilCompleteHttpFixture, "Can early wake up when adding n
 	HttpRequestFast->ProcessRequest();
 }
 
+#if HAS_EVENTLOOP_NOTIFIER_EVENTFD
+class FPerfAsyncEventLoopNotifierHttpFixture : public FHttpModuleTestFixture
+{
+public:
+	FPerfAsyncEventLoopNotifierHttpFixture()
+	{
+		// For FEventLoopNotifierEventFd::Notify to use thread pool in WebTests
+		CHECK(GIOThreadPool == nullptr);
+		GIOThreadPool = FQueuedThreadPool::Allocate();
+		int32 NumThreadsInThreadPool = FPlatformMisc::NumberOfIOWorkerThreadsToSpawn();
+		REQUIRE(GIOThreadPool->Create(NumThreadsInThreadPool, 96 * 1024, TPri_AboveNormal, TEXT("IOThreadPool")));
+	}
+
+	~FPerfAsyncEventLoopNotifierHttpFixture()
+	{
+		GIOThreadPool->Destroy();
+		GIOThreadPool = nullptr;
+	}
+};
+
+TEST_CASE_METHOD(FPerfAsyncEventLoopNotifierHttpFixture, "Can use event loop notify with good performance", HTTP_TAG)
+#else // HAS_EVENTLOOP_NOTIFIER_EVENTFD
+TEST_CASE_METHOD(FHttpModuleTestFixture, "Can use event loop notify with good performance", HTTP_TAG)
+#endif // HAS_EVENTLOOP_NOTIFIER_EVENTFD
+{
+	DisableWarningsInThisTest(); // for: [FHttpManager::Shutdown] Unbinding delegates for 5 outstanding Http Requests
+
+	for (int32 i = 0; i < 20; ++i)
+	{
+		TSharedRef<IHttpRequest> HttpRequest = FHttpModule::Get().CreateRequest();
+		HttpRequest->SetURL(UrlStreamDownload(3, HTTP_TEST_TIMEOUT_CHUNK_SIZE, /*ChunkLatency=*/ 0));
+		HttpRequest->SetDelegateThreadPolicy(EHttpRequestDelegateThreadPolicy::CompleteOnHttpThread);
+
+		HttpRequest->OnProcessRequestComplete().BindLambda([](FHttpRequestPtr HttpRequest, FHttpResponsePtr HttpResponse, bool bSucceeded)
+		{
+			CHECK(bSucceeded);
+			REQUIRE(HttpResponse != nullptr);
+			CHECK(HttpResponse->GetResponseCode() == 200);
+		});
+
+		const double StartTime = FPlatformTime::Seconds();
+		HttpRequest->ProcessRequest();
+		const double DurationInSeconds  = FPlatformTime::Seconds() - StartTime;
+		REQUIRE(DurationInSeconds < 0.001); // It shouldn't exceed 1ms on any platform
+
+		FPlatformProcess::Sleep(0.005f);
+	}
+}
+
 TEST_CASE_METHOD(FWaitUntilCompleteHttpFixture, "Can run a lot of small download requests with good performance", HTTP_TAG)
 {
 	if (!bRunHeavyTests)
@@ -1158,7 +1208,6 @@ TEST_CASE_METHOD(FWaitUntilCompleteHttpFixture, "Can run a lot of small download
 	const double StartTime = FPlatformTime::Seconds();
 
 	int32 NumRequests = 10000;
-	FEvent* Event = FPlatformProcess::GetSynchEventFromPool(true);
 
 	std::atomic<int32> RunningRequests = 0;
 	for (int32 i = 0; i < NumRequests; ++i)
@@ -1167,7 +1216,7 @@ TEST_CASE_METHOD(FWaitUntilCompleteHttpFixture, "Can run a lot of small download
 		HttpRequest->SetURL(UrlStreamDownload(3, HTTP_TEST_TIMEOUT_CHUNK_SIZE, /*ChunkLatency=*/ 0));
 		HttpRequest->SetDelegateThreadPolicy(EHttpRequestDelegateThreadPolicy::CompleteOnHttpThread);
 
-		HttpRequest->OnProcessRequestComplete().BindLambda([Event, &RunningRequests](FHttpRequestPtr HttpRequest, FHttpResponsePtr HttpResponse, bool bSucceeded)
+		HttpRequest->OnProcessRequestComplete().BindLambda([&RunningRequests](FHttpRequestPtr HttpRequest, FHttpResponsePtr HttpResponse, bool bSucceeded)
 		{
 			CHECK(bSucceeded);
 			REQUIRE(HttpResponse != nullptr);
diff --git a/Engine/Source/Programs/WebTests/WebTests.Build.cs b/Engine/Source/Programs/WebTests/WebTests.Build.cs
--- a/Engine/Source/Programs/WebTests/WebTests.Build.cs
+++ b/Engine/Source/Programs/WebTests/WebTests.Build.cs
@@ -35,6 +35,7 @@ public class WebTests : TestModuleRules
 			new string[] {
 				"ApplicationCore",
 				"Core",
+				"EventLoop",
 				"HTTP",
 				"SSL",
 				"HTTPServer",
diff --git a/Engine/Source/Runtime/Online/Experimental/EventLoop/Private/EventLoop/BSDSocket/EventLoopNotifierEventFd.cpp b/Engine/Source/Runtime/Online/Experimental/EventLoop/Private/EventLoop/BSDSocket/EventLoopNotifierEventFd.cpp
--- a/Engine/Source/Runtime/Online/Experimental/EventLoop/Private/EventLoop/BSDSocket/EventLoopNotifierEventFd.cpp
+++ b/Engine/Source/Runtime/Online/Experimental/EventLoop/Private/EventLoop/BSDSocket/EventLoopNotifierEventFd.cpp
@@ -2,9 +2,25 @@
 
 #include "EventLoop/BSDSocket/EventLoopNotifierEventFd.h"
 
+#include "Async/Async.h"
+#include "EventLoop/EventLoopLog.h"
+#include "HAL//IConsoleManager.h"
+#include "Misc/QueuedThreadPool.h"
+#include "Stats/Stats.h"
+
 #if HAS_EVENTLOOP_NOTIFIER_EVENTFD
 #include "BSDSocketTypesPrivate.h"
 
+namespace UE::EventLoop {
+
+static bool bHttpEventLoopNotifierUseThreadPool = true;
+static FAutoConsoleVariableRef CVarHttpEventLoopNotifierUseThreadPool(
+	TEXT("http.EventLoopNotifierUseThreadPool"),
+	bHttpEventLoopNotifierUseThreadPool,
+	TEXT("When enabled, in event loop it will use thread pool to notify other thread to wake up"),
+	ECVF_SaveForNextBoot
+);
+
 bool FEventLoopNotifierEventFd::Init()
 {
 	FileDescriptorRead = eventfd(0, EFD_NONBLOCK);
@@ -20,13 +36,36 @@ void FEventLoopNotifierEventFd::Shutdown()
 	}
 }
 
-void FEventLoopNotifierEventFd::Notify()
+namespace EventLoopNotifierEventFd
+{
+
+static void NotifyImpl(int32 Fd)
 {
 	eventfd_t Value = 1;
-	if (eventfd_write(FileDescriptorRead, Value) == SOCKET_RESULT_FAILED)
+	if (eventfd_write(Fd, Value) == SOCKET_RESULT_FAILED)
 	{
 		int Err = errno;
-		check(Err == EAGAIN);
+		UE_CLOG(Err != EAGAIN, LogEventLoop, Warning, TEXT("Failed to notify by event_fd_write %d with error %d"), Fd, Err);
+	}
+}
+
+}
+
+void FEventLoopNotifierEventFd::Notify()
+{
+	QUICK_SCOPE_CYCLE_COUNTER(STAT_FEventLoopNotifierEventFd_Notify);
+
+#if PLATFORM_WRITES_ARE_SLOW
+	if (bHttpEventLoopNotifierUseThreadPool && GIOThreadPool)
+	{
+		AsyncPool(*GIOThreadPool, [Fd = FileDescriptorRead] {
+			EventLoopNotifierEventFd::NotifyImpl(Fd);
+		});
+	}
+	else
+#endif
+	{
+		EventLoopNotifierEventFd::NotifyImpl(FileDescriptorRead);
 	}
 }
 
@@ -36,5 +75,7 @@ void FEventLoopNotifierEventFd::Clear()
 	ensure(eventfd_read(FileDescriptorRead, &Value) != SOCKET_RESULT_FAILED);
 }
 
+/* UE::EventLoop */ }
+
 #endif // HAS_EVENTLOOP_NOTIFIER_EVENTFD
 
diff --git a/Engine/Source/Runtime/Online/Experimental/EventLoop/Private/EventLoop/BSDSocket/EventLoopNotifierSocket.cpp b/Engine/Source/Runtime/Online/Experimental/EventLoop/Private/EventLoop/BSDSocket/EventLoopNotifierSocket.cpp
--- a/Engine/Source/Runtime/Online/Experimental/EventLoop/Private/EventLoop/BSDSocket/EventLoopNotifierSocket.cpp
+++ b/Engine/Source/Runtime/Online/Experimental/EventLoop/Private/EventLoop/BSDSocket/EventLoopNotifierSocket.cpp
@@ -2,10 +2,13 @@
 
 #include "EventLoop/BSDSocket/EventLoopNotifierSocket.h"
 #include "BSDSocketTypesPrivate.h"
+#include "Stats/Stats.h"
 
 #if PLATFORM_HAS_BSD_SOCKETS && !PLATFORM_HAS_BSD_SOCKET_FEATURE_WINSOCKETS
 
-namespace UE::EventLoopSocketNotifierInternal
+namespace UE::EventLoop {
+
+namespace EventLoopNotifierSocket
 {
 
 static void SetSocketNonBlock(int32 FileDescriptor)
@@ -48,7 +51,7 @@ bool FEventLoopNotifierSocket::Init()
 		return false;
 	}
 
-	UE::EventLoopSocketNotifierInternal::SetSocketNonBlock(FileDescriptorListener);
+	EventLoopNotifierSocket::SetSocketNonBlock(FileDescriptorListener);
 
 
 
@@ -77,7 +80,7 @@ bool FEventLoopNotifierSocket::Init()
 		return false;
 	}
 
-	UE::EventLoopSocketNotifierInternal::SetSocketNonBlock(FileDescriptorRead);
+	EventLoopNotifierSocket::SetSocketNonBlock(FileDescriptorRead);
 
 
 
@@ -94,7 +97,7 @@ bool FEventLoopNotifierSocket::Init()
 		return false;
 	}
 
-	UE::EventLoopSocketNotifierInternal::SetSocketNonBlock(FileDescriptorWrite);
+	EventLoopNotifierSocket::SetSocketNonBlock(FileDescriptorWrite);
 
 
 
@@ -125,6 +128,7 @@ void FEventLoopNotifierSocket::Shutdown()
 
 void FEventLoopNotifierSocket::Notify()
 {
+	QUICK_SCOPE_CYCLE_COUNTER(STAT_FEventLoopNotifierSocket_Notify);
 	if (FileDescriptorWrite != INVALID_SOCKET)
 	{
 		// TODO: A failure to write may not be critical. Its possible that the receive buffer is just full. 
@@ -148,5 +152,7 @@ void FEventLoopNotifierSocket::Clear()
 	}
 }
 
+/* UE::EventLoop */ }
+
 #endif // PLATFORM_HAS_BSD_SOCKETS && !PLATFORM_HAS_BSD_SOCKET_FEATURE_WINSOCKETS
 
diff --git a/Engine/Source/Runtime/Online/Experimental/EventLoop/Public/EventLoop/BSDSocket/EventLoopNotifier.h b/Engine/Source/Runtime/Online/Experimental/EventLoop/Public/EventLoop/BSDSocket/EventLoopNotifier.h
--- a/Engine/Source/Runtime/Online/Experimental/EventLoop/Public/EventLoop/BSDSocket/EventLoopNotifier.h
+++ b/Engine/Source/Runtime/Online/Experimental/EventLoop/Public/EventLoop/BSDSocket/EventLoopNotifier.h
@@ -6,8 +6,14 @@
 
 #if HAS_EVENTLOOP_NOTIFIER_EVENTFD
 #include "EventLoop/BSDSocket/EventLoopNotifierEventFd.h"
+namespace UE::EventLoop
+{
 using FEventLoopNotifier = FEventLoopNotifierEventFd;
+}
 #elif PLATFORM_HAS_BSD_SOCKETS && !PLATFORM_HAS_BSD_SOCKET_FEATURE_WINSOCKETS
 #include "EventLoop/BSDSocket/EventLoopNotifierSocket.h"
+namespace UE::EventLoop
+{
 using FEventLoopNotifier = FEventLoopNotifierSocket;
+}
 #endif
diff --git a/Engine/Source/Runtime/Online/Experimental/EventLoop/Public/EventLoop/BSDSocket/EventLoopNotifierEventFd.h b/Engine/Source/Runtime/Online/Experimental/EventLoop/Public/EventLoop/BSDSocket/EventLoopNotifierEventFd.h
--- a/Engine/Source/Runtime/Online/Experimental/EventLoop/Public/EventLoop/BSDSocket/EventLoopNotifierEventFd.h
+++ b/Engine/Source/Runtime/Online/Experimental/EventLoop/Public/EventLoop/BSDSocket/EventLoopNotifierEventFd.h
@@ -7,6 +7,8 @@
 
 #if HAS_EVENTLOOP_NOTIFIER_EVENTFD
 
+namespace UE::EventLoop {
+
 class FEventLoopNotifierEventFd
 {
 public:
@@ -22,4 +24,6 @@ private:
 	int32 FileDescriptorRead = INVALID_SOCKET;
 };
 
+/* UE::EventLoop */ }
+
 #endif // HAS_EVENTLOOP_NOTIFIER_EVENTFD
diff --git a/Engine/Source/Runtime/Online/Experimental/EventLoop/Public/EventLoop/BSDSocket/EventLoopNotifierSocket.h b/Engine/Source/Runtime/Online/Experimental/EventLoop/Public/EventLoop/BSDSocket/EventLoopNotifierSocket.h
--- a/Engine/Source/Runtime/Online/Experimental/EventLoop/Public/EventLoop/BSDSocket/EventLoopNotifierSocket.h
+++ b/Engine/Source/Runtime/Online/Experimental/EventLoop/Public/EventLoop/BSDSocket/EventLoopNotifierSocket.h
@@ -7,6 +7,8 @@
 
 #if PLATFORM_HAS_BSD_SOCKETS && !PLATFORM_HAS_BSD_SOCKET_FEATURE_WINSOCKETS
 
+namespace UE::EventLoop {
+
 class FEventLoopNotifierSocket
 {
 public:
@@ -24,5 +26,7 @@ private:
 	int32 FileDescriptorRead = INVALID_SOCKET;
 };
 
+/* UE::EventLoop */ }
+
 #endif // PLATFORM_HAS_BSD_SOCKETS && !PLATFORM_HAS_BSD_SOCKET_FEATURE_WINSOCKETS
 
diff --git a/Engine/Source/Runtime/Online/Experimental/EventLoop/Public/EventLoop/EventLoop.h b/Engine/Source/Runtime/Online/Experimental/EventLoop/Public/EventLoop/EventLoop.h
--- a/Engine/Source/Runtime/Online/Experimental/EventLoop/Public/EventLoop/EventLoop.h
+++ b/Engine/Source/Runtime/Online/Experimental/EventLoop/Public/EventLoop/EventLoop.h
@@ -301,6 +301,8 @@ private:
 
 	void TryNotify()
 	{
+		QUICK_SCOPE_CYCLE_COUNTER(STAT_TEventLoop_TryNotify);
+
 		if (Traits::IsInitialized(EventLoopThreadId))
 		{
 			// Init already completed. Notify the IO manager.
diff --git a/Engine/Source/Runtime/Online/HTTP/Private/Curl/CurlMultiWaitEventLoopHttpThread.h b/Engine/Source/Runtime/Online/HTTP/Private/Curl/CurlMultiWaitEventLoopHttpThread.h
--- a/Engine/Source/Runtime/Online/HTTP/Private/Curl/CurlMultiWaitEventLoopHttpThread.h
+++ b/Engine/Source/Runtime/Online/HTTP/Private/Curl/CurlMultiWaitEventLoopHttpThread.h
@@ -68,7 +68,7 @@ private:
 	UE::EventLoop::IEventLoop& EventLoop;
 	FParams Params;
 	int32 EmptySequentialWaitCount;
-	FEventLoopNotifier EventLoopNotifier;
+	UE::EventLoop::FEventLoopNotifier EventLoopNotifier;
 };
 
 class FCurlMultiWaitEventLoopHttpThread
diff --git a/Engine/Source/Runtime/Online/HTTP/Private/EventLoopHttpThread.cpp b/Engine/Source/Runtime/Online/HTTP/Private/EventLoopHttpThread.cpp
--- a/Engine/Source/Runtime/Online/HTTP/Private/EventLoopHttpThread.cpp
+++ b/Engine/Source/Runtime/Online/HTTP/Private/EventLoopHttpThread.cpp
@@ -103,6 +103,7 @@ void FEventLoopHttpThread::AddRequest(FHttpRequestCommon* Request)
 {
 	FHttpThreadBase::AddRequest(Request);
 
+	QUICK_SCOPE_CYCLE_COUNTER(STAT_FEventLoopHttpThread_AddRequest);
 	if (UE::EventLoop::IEventLoop* EventLoop = GetEventLoop())
 	{
 		// Force a wakeup to process new tasks.
