Add a new websocket test case: Make sure WebSockets can respond the Close from websockets server.

#Synced-CL 49840330

diff --git a/Engine/Source/Programs/WebTests/Private/TestWebSockets.cpp b/Engine/Source/Programs/WebTests/Private/TestWebSockets.cpp
--- a/Engine/Source/Programs/WebTests/Private/TestWebSockets.cpp
+++ b/Engine/Source/Programs/WebTests/Private/TestWebSockets.cpp
@@ -150,6 +150,27 @@ TEST_CASE_METHOD(FRunUntilQuitRequestedFixture, "WebSockets can connect then sen
 	WebSocket->Connect();
 }
 
+TEST_CASE_METHOD(FRunUntilQuitRequestedFixture, "WebSockets can be closed by server", WEBSOCKETS_TAG)
+{
+	TSharedPtr<IWebSocket> WebSocket = WebSocketsModule->CreateWebSocket(FString::Format(TEXT("{0}/close_on_receive_message/"), { *UrlWebSocketsTests() }));
+
+	WebSocket->OnConnected().AddLambda([this, WebSocket](){
+		WebSocket->Send(TEXT("close me"));
+	});
+
+	WebSocket->OnClosed().AddLambda([this, WebSocket](int32 StatusCode, const FString& /* Reason */, bool /* bWasClean */){
+		CHECK(StatusCode == 1000); // 1000 means default normal closure
+		bQuitRequested = true;
+	});
+
+	WebSocket->OnConnectionError().AddLambda([this, WebSocket](const FString& /* Error */){
+		CHECK(false);
+		bQuitRequested = true;
+	});
+
+	WebSocket->Connect();
+}
+
 TEST_CASE_METHOD(FRunUntilQuitRequestedFixture, "WebSockets module can shut down when there are still websockets connections", WEBSOCKETS_TAG)
 {
 	DisableWarningsInThisTest();
diff --git a/Engine/Source/Programs/WebTestsServer/websocketstests/consumers.py b/Engine/Source/Programs/WebTestsServer/websocketstests/consumers.py
--- a/Engine/Source/Programs/WebTestsServer/websocketstests/consumers.py
+++ b/Engine/Source/Programs/WebTestsServer/websocketstests/consumers.py
@@ -11,3 +11,13 @@ class EchoConsumer(WebsocketConsumer):
 
     def receive(self, text_data):
         self.send(text_data)
+
+class CloseOnReceiveMessageConsumer(WebsocketConsumer):
+    def connect(self):
+        self.accept()
+
+    def disconnect(self, close_code):
+        pass
+
+    def receive(self, text_data):
+        self.close()
diff --git a/Engine/Source/Programs/WebTestsServer/websocketstests/routing.py b/Engine/Source/Programs/WebTestsServer/websocketstests/routing.py
--- a/Engine/Source/Programs/WebTestsServer/websocketstests/routing.py
+++ b/Engine/Source/Programs/WebTestsServer/websocketstests/routing.py
@@ -4,4 +4,5 @@ from . import consumers
 
 websocket_urlpatterns = [
     re_path(r"webtests/websocketstests/echo", consumers.EchoConsumer.as_asgi()),
+    re_path(r"webtests/websocketstests/close_on_receive_message", consumers.CloseOnReceiveMessageConsumer.as_asgi()),
 ]
