[HTTP] Add TakeContent() method to HttpResponse that uses move semantics to hand you the contents, letting you avoid copies when you know you are the only user of the content.

#tests mainstage through launchit, webtest.

#Synced-CL 51332317

diff --git a/Engine/Source/Programs/WebTests/Private/TestHttp.cpp b/Engine/Source/Programs/WebTests/Private/TestHttp.cpp
--- a/Engine/Source/Programs/WebTests/Private/TestHttp.cpp
+++ b/Engine/Source/Programs/WebTests/Private/TestHttp.cpp
@@ -58,6 +58,25 @@ public:
 	using FHttpModule::HttpActivityTimeout;
 };
 
+namespace HttpTestUtils
+{
+	// checker.
+	static void CheckContent(int DataLength, int RepeatAt, const TArray<uint8>& Content)
+	{
+		CHECK(Content.Num() == DataLength);
+
+		bool bAllMatch = true;
+
+		// Make sure the data got read is in good state
+		for (int32 i = 0; i < Content.Num(); ++i)
+		{
+			bAllMatch &= (Content[i] == '0' + (i % RepeatAt));
+		}
+
+		CHECK(bAllMatch);
+	};
+};
+
 class FHttpTestLogLevelInitializer
 {
 public:
@@ -679,19 +698,14 @@ TEST_CASE_METHOD(FWaitUntilCompleteHttpFixture, "Can process request with gzip s
 		CHECK(bSucceeded);
 		REQUIRE(HttpResponse != nullptr);
 		CHECK(HttpResponse->GetResponseCode() == 200);
+		
+		// Using get response.
+		HttpTestUtils::CheckContent(DataLength, RepeatAt, HttpResponse->GetContent());
 
-		const TArray<uint8>& Content = HttpResponse->GetContent();
-		CHECK(Content.Num() == DataLength);
-
-		bool bAllMatch = true;
-
-		// Make sure the data got read is in good state
-		for (int32 i = 0; i < Content.Num(); ++i)
-		{
-			bAllMatch &= (Content[i] == '0' + (i % RepeatAt));
-		}
-
-		CHECK(bAllMatch);
+		// The most we can test is that it return the same thing, since
+		// implementations may not be able to actually do the move semantics. 
+		TArray<uint8> TakenContent = HttpResponse->TakeContent();
+		HttpTestUtils::CheckContent(DataLength, RepeatAt, TakenContent);
 	});
 	HttpRequest->ProcessRequest();
 }
@@ -840,18 +854,9 @@ TEST_CASE_METHOD(FWaitUntilCompleteHttpFixture, "Get large response content with
 		REQUIRE(HttpResponse != nullptr);
 		CHECK(HttpResponse->GetResponseCode() == 200);
 
-		const TArray<uint8>& Content = HttpResponse->GetContent();
-		CHECK(Content.Num() == DataLength);
-
-		bool bAllMatch = true;
-
-		// Make sure the data got read is in good state
-		for (int32 i = 0; i < Content.Num(); ++i)
-		{
-			bAllMatch &= (Content[i] == '0' + (i % RepeatAt));
-		}
-
-		CHECK(bAllMatch);
+		HttpTestUtils::CheckContent(DataLength, RepeatAt, HttpResponse->GetContent());
+		TArray<uint8> TakenContent = HttpResponse->TakeContent();
+		HttpTestUtils::CheckContent(DataLength, RepeatAt, TakenContent);
 	});
 	HttpRequest->ProcessRequest();
 }
diff --git a/Engine/Source/Runtime/Online/BuildPatchServices/Private/Tests/Mock/HttpResponse.mock.h b/Engine/Source/Runtime/Online/BuildPatchServices/Private/Tests/Mock/HttpResponse.mock.h
--- a/Engine/Source/Runtime/Online/BuildPatchServices/Private/Tests/Mock/HttpResponse.mock.h
+++ b/Engine/Source/Runtime/Online/BuildPatchServices/Private/Tests/Mock/HttpResponse.mock.h
@@ -77,6 +77,12 @@ namespace BuildPatchServices
 			return None;
 		}
 
+		virtual TArray<uint8> TakeContent() override
+		{
+			MOCK_FUNC_NOT_IMPLEMENTED("FMockHttpRequest::TakeContent");
+			return {};
+		}
+
 		virtual EHttpRequestStatus::Type GetStatus() const override
 		{
 			MOCK_FUNC_NOT_IMPLEMENTED("FMockHttpRequest::GetStatus");
diff --git a/Engine/Source/Runtime/Online/HTTP/Private/GenericPlatform/HttpResponseCommon.cpp b/Engine/Source/Runtime/Online/HTTP/Private/GenericPlatform/HttpResponseCommon.cpp
--- a/Engine/Source/Runtime/Online/HTTP/Private/GenericPlatform/HttpResponseCommon.cpp
+++ b/Engine/Source/Runtime/Online/HTTP/Private/GenericPlatform/HttpResponseCommon.cpp
@@ -100,6 +100,15 @@ const TArray<uint8>& FHttpResponseCommon::GetContent() const
 	return Payload;
 }
 
+TArray<uint8> FHttpResponseCommon::TakeContent()
+{
+	if (!bIsReady)
+	{
+		UE_LOG(LogHttp, Warning, TEXT("Payload is incomplete. Response still processing for %s"), *GetURL());
+	}
+	return MoveTemp(Payload);
+}
+
 FString FHttpResponseCommon::GetContentAsString() const
 {
 	// Content is NOT null-terminated; we need to specify lengths here
diff --git a/Engine/Source/Runtime/Online/HTTP/Private/GenericPlatform/HttpResponseCommon.h b/Engine/Source/Runtime/Online/HTTP/Private/GenericPlatform/HttpResponseCommon.h
--- a/Engine/Source/Runtime/Online/HTTP/Private/GenericPlatform/HttpResponseCommon.h
+++ b/Engine/Source/Runtime/Online/HTTP/Private/GenericPlatform/HttpResponseCommon.h
@@ -26,6 +26,7 @@ public:
 	virtual FUtf8StringView GetContentAsUtf8StringView() const override;
 	virtual FString GetHeader(const FString& HeaderName) const override;
 	virtual const TArray<uint8>& GetContent() const override;
+	virtual TArray<uint8> TakeContent() override;
 	virtual FString GetContentAsString() const override;
 
 protected:
diff --git a/Engine/Source/Runtime/Online/HTTP/Public/Interfaces/IHttpResponse.h b/Engine/Source/Runtime/Online/HTTP/Public/Interfaces/IHttpResponse.h
--- a/Engine/Source/Runtime/Online/HTTP/Public/Interfaces/IHttpResponse.h
+++ b/Engine/Source/Runtime/Online/HTTP/Public/Interfaces/IHttpResponse.h
@@ -110,7 +110,13 @@ namespace EHttpResponseCodes
 class IHttpResponse : public IHttpBase
 {
 public:
-
+	/**
+	* Return the content with a MoveTemp if possible.
+	*
+	* @return - The content array
+	*/
+	virtual TArray<uint8> TakeContent() = 0;
+	
 	/**
 	 * Gets the response code returned by the requested server.
 	 * See EHttpResponseCodes for known response codes
