Fix the http server response error when enable bEnableInstrumentation in build config.

#Synced-CL 50586490

diff --git a/Engine/Source/Programs/WebTests/Private/TestHttp.cpp b/Engine/Source/Programs/WebTests/Private/TestHttp.cpp
--- a/Engine/Source/Programs/WebTests/Private/TestHttp.cpp
+++ b/Engine/Source/Programs/WebTests/Private/TestHttp.cpp
@@ -3211,6 +3211,39 @@ TEST_CASE_METHOD(FLocalHttpServerFixture, "Local http server can serve large fil
 	HttpRequest->ProcessRequest();
 }
 
+TEST_CASE_METHOD(FLocalHttpServerFixture, "Local http server can handle post request", HTTP_TAG)
+{
+	StartServerWithHandler(FHttpPath(TEXT("/post_path")), EHttpServerRequestVerbs::VERB_POST, FHttpRequestHandler::CreateLambda([](const FHttpServerRequest& Request, const FHttpResultCallback& OnComplete) {
+		FString JsonString;
+		TSharedRef<TJsonWriter<>> JsonWriter = TJsonWriterFactory<>::Create(&JsonString);
+		JsonWriter->WriteObjectStart();
+		JsonWriter->WriteValue("success", false);
+		JsonWriter->WriteObjectEnd();
+		JsonWriter->Close();
+
+		TUniquePtr<FHttpServerResponse> Response = FHttpServerResponse::Create(JsonString, TEXT("application/json"));
+		OnComplete(MoveTemp(Response));
+
+		return true;
+	}));
+
+	uint32 ChunkSize = 128;
+	TArray<uint8> DataChunk;
+	DataChunk.SetNum(ChunkSize);
+	FMemory::Memset(DataChunk.GetData(), 'd', ChunkSize);
+
+	// Start client request
+	TSharedRef<IHttpRequest> HttpRequest = CreateRequest();
+	HttpRequest->SetURL(TEXT("http://localhost:9000/post_path"));
+	HttpRequest->SetVerb(FHttpConstants::VERB_POST);
+	HttpRequest->SetContent(DataChunk);
+	HttpRequest->OnProcessRequestComplete().BindLambda([](FHttpRequestPtr HttpRequest, FHttpResponsePtr HttpResponse, bool bSucceeded) {
+		CHECK(bSucceeded);
+		REQUIRE(HttpResponse != nullptr);
+	});
+	HttpRequest->ProcessRequest();
+}
+
 #endif // UE_HTTP_SUPPORT_LOCAL_SERVER
 
 TEST_CASE_METHOD(FWaitUntilCompleteHttpFixture, "Accessing request initial information without issue while request is running", HTTP_TAG)
diff --git a/Engine/Source/Runtime/Online/HTTPServer/Private/HttpConnectionResponseWriteContext.cpp b/Engine/Source/Runtime/Online/HTTPServer/Private/HttpConnectionResponseWriteContext.cpp
--- a/Engine/Source/Runtime/Online/HTTPServer/Private/HttpConnectionResponseWriteContext.cpp
+++ b/Engine/Source/Runtime/Online/HTTPServer/Private/HttpConnectionResponseWriteContext.cpp
@@ -126,8 +126,7 @@ bool FHttpConnectionResponseWriteContext::IsWriteBodyComplete() const
 
 TArray<uint8> FHttpConnectionResponseWriteContext::SerializeHeadersUtf8(HttpVersion::EHttpServerHttpVersion HttpVersion, EHttpServerResponseCodes ResponseCode, const TMap<FString, TArray<FString>>& HeadersMap)
 {
-	FString ResponseHeaderStr = FString::Printf(TEXT("%s %d\r\n"),  
-		*HttpVersion::ToString(HttpVersion), ResponseCode);
+	FString ResponseHeaderStr = FString::Printf(TEXT("%s %d\r\n"),  *HttpVersion::ToString(HttpVersion), static_cast<int16>(ResponseCode));
 
 	for (const auto& KeyValuePair : HeadersMap)
 	{
