Fix the error in WebTests when quickly listen 9000 again on linux.

#Synced-CL 50586490

diff --git a/Engine/Source/Programs/WebTests/Private/TestHttp.cpp b/Engine/Source/Programs/WebTests/Private/TestHttp.cpp
--- a/Engine/Source/Programs/WebTests/Private/TestHttp.cpp
+++ b/Engine/Source/Programs/WebTests/Private/TestHttp.cpp
@@ -16,6 +16,7 @@
 #include "HttpServerModule.h"
 #include "Logging/LogScopedVerbosityOverride.h"
 #include "Misc/CommandLine.h"
+#include "Misc/ConfigCacheIni.h"
 #include "Misc/QueuedThreadPool.h"
 #include "Misc/ScopeExit.h"
 #include "TestHarness.h"
@@ -3242,11 +3243,17 @@ TEST_CASE_METHOD(FWaitUntilCompleteHttpFixture, "Scheme besides http and https c
 	HttpRequest->ProcessRequest();
 }
 
+extern TAutoConsoleVariable<bool> CVarHttpServerDefaultReuseAddressAndPortEnabled;
+
 class FLocalHttpServerFixture : public FWaitUntilCompleteHttpFixture
 {
 public:
 	FLocalHttpServerFixture()
 	{
+		// It could fail to bind the same port between test cases, when previous test is still cleaning up on OS layer
+		// So make sure it can reuse the same address and port
+		CVarHttpServerDefaultReuseAddressAndPortEnabled->Set(true);
+
 		HttpServerModule = new FHttpServerModule();
 		HttpServerModule->StartupModule();
 
@@ -3297,7 +3304,7 @@ TEST_CASE_METHOD(FLocalHttpServerFixture, "Local http server can serve large fil
 
 	// Start client request
 	TSharedRef<IHttpRequest> HttpRequest = CreateRequest();
-	HttpRequest->SetURL(TEXT("http://localhost:9000/large_file"));
+	HttpRequest->SetURL(FString::Format(TEXT("http://localhost:{0}/large_file"), {LocalHttpServerPort}));
 	HttpRequest->SetVerb(TEXT("GET"));
 	HttpRequest->OnProcessRequestComplete().BindLambda([FileSize](FHttpRequestPtr HttpRequest, FHttpResponsePtr HttpResponse, bool bSucceeded) {
 		CHECK(bSucceeded);
@@ -3330,7 +3337,7 @@ TEST_CASE_METHOD(FLocalHttpServerFixture, "Local http server can handle post req
 
 	// Start client request
 	TSharedRef<IHttpRequest> HttpRequest = CreateRequest();
-	HttpRequest->SetURL(TEXT("http://localhost:9000/post_path"));
+	HttpRequest->SetURL(FString::Format(TEXT("http://localhost:{0}/post_path"), {LocalHttpServerPort}));
 	HttpRequest->SetVerb(FHttpConstants::VERB_POST);
 	HttpRequest->SetContent(DataChunk);
 	HttpRequest->OnProcessRequestComplete().BindLambda([](FHttpRequestPtr HttpRequest, FHttpResponsePtr HttpResponse, bool bSucceeded) {
diff --git a/Engine/Source/Runtime/Online/HTTPServer/Private/HttpServerConfig.cpp b/Engine/Source/Runtime/Online/HTTPServer/Private/HttpServerConfig.cpp
--- a/Engine/Source/Runtime/Online/HTTPServer/Private/HttpServerConfig.cpp
+++ b/Engine/Source/Runtime/Online/HTTPServer/Private/HttpServerConfig.cpp
@@ -4,9 +4,17 @@
 #include "CoreMinimal.h"
 #include "Misc/ConfigCacheIni.h"
 #include "Misc/CoreDelegates.h"
+#include "HAL/IConsoleManager.h"
 
 DEFINE_LOG_CATEGORY(LogHttpServerConfig)
 
+TAutoConsoleVariable<bool> CVarHttpServerDefaultReuseAddressAndPortEnabled(
+	TEXT("HTTPServer.DefaultReuseAddressAndPortEnabled"),
+	false,
+	TEXT("When enabled, http server will reuse address and port without raising error")
+);
+
+// TODO: static FString will cause mem leaks warnings in Insights, try to use macro or member ptr in http server module instead
 namespace FHttpServerConfigCache
 {
 	static FDelegateHandle ConfigSectionChangedHandle;
@@ -53,7 +61,10 @@ const FHttpServerListenerConfig FHttpServerConfig::GetListenerConfig(uint32 Port
 		GConfig->GetInt(*IniSectionNameHTTPServerListeners, TEXT("DefaultBufferSize"), CacheListenerConfig.BufferSize, GEngineIni);
 		GConfig->GetInt(*IniSectionNameHTTPServerListeners, TEXT("DefaultConnectionsBacklogSize"), CacheListenerConfig.ConnectionsBacklogSize, GEngineIni);
 		GConfig->GetInt(*IniSectionNameHTTPServerListeners, TEXT("DefaultMaxConnectionsAcceptPerFrame"), CacheListenerConfig.MaxConnectionsAcceptPerFrame, GEngineIni);
-		GConfig->GetBool(*IniSectionNameHTTPServerListeners, TEXT("DefaultReuseAddressAndPort"), CacheListenerConfig.bReuseAddressAndPort, GEngineIni);
+		if (!GConfig->GetBool(*IniSectionNameHTTPServerListeners, TEXT("DefaultReuseAddressAndPort"), CacheListenerConfig.bReuseAddressAndPort, GEngineIni))
+		{
+			CacheListenerConfig.bReuseAddressAndPort = CVarHttpServerDefaultReuseAddressAndPortEnabled.GetValueOnAnyThread();
+		}
 
 		GConfig->GetArray(*IniSectionNameHTTPServerListeners, TEXT("ListenerOverrides"), CacheListenerConfigs, GEngineIni);
 
@@ -111,4 +122,4 @@ const FHttpServerConnectionConfig FHttpServerConfig::GetConnectionConfig()
 	}
 
 	return CacheConnectionConfig;
-}
\ No newline at end of file
+}
