Fix the test failure for xCurl: the recv buffer has increased in the gdk update, we need to increase the chunk size in order to trigger data received callback immediately.

#Synced-CL 51366129

diff --git a/Engine/Source/Programs/WebTests/Private/TestHttp.cpp b/Engine/Source/Programs/WebTests/Private/TestHttp.cpp
--- a/Engine/Source/Programs/WebTests/Private/TestHttp.cpp
+++ b/Engine/Source/Programs/WebTests/Private/TestHttp.cpp
@@ -41,7 +41,9 @@
 
 #define HTTP_TAG "[HTTP]"
 #define HTTP_TIME_DIFF_TOLERANCE_OF_REQUEST 0.5f
-#define HTTP_TEST_TIMEOUT_CHUNK_SIZE 16*1024 // Use a big chunk size so it triggers data received callback in time on all platforms
+ // Use a big chunk size so it triggers data received callback in time on all platforms
+ // For xCurl, this is the minimum chunk size to trigger data received callback immediately, without waiting additional time for the recv buffer to fill up
+#define HTTP_TEST_TIMEOUT_CHUNK_SIZE 32*1024
 
 extern TAutoConsoleVariable<bool> CVarHttpInsecureProtocolEnabled;
 extern TAutoConsoleVariable<int32> CVarHttpMaxConcurrentRequests;
