[Http] Add per-failure-reason stats to HttpManager.

#Synced-CL 51774074

#Synced-CL 51774074

diff --git a/Engine/Source/Programs/WebTests/Private/TestHttp.cpp b/Engine/Source/Programs/WebTests/Private/TestHttp.cpp
--- a/Engine/Source/Programs/WebTests/Private/TestHttp.cpp
+++ b/Engine/Source/Programs/WebTests/Private/TestHttp.cpp
@@ -3451,6 +3451,42 @@ TEST_CASE_METHOD(FWaitUntilCompleteHttpFixture, "Test platform request requests
 	}
 }
 
+TEST_CASE_METHOD(FWaitUntilCompleteHttpFixture, "RecordStatRequest increments RequestResults for None and ConnectionError", HTTP_TAG)
+{
+	DisableWarningsInThisTest();
+
+	TSharedRef<IHttpRequest> HttpRequest = CreateRequest();
+
+	SECTION("Succeeded request increments RequestResults[None]")
+	{
+		HttpRequest->SetURL(UrlToTestMethods());
+		HttpRequest->OnProcessRequestComplete().BindLambda([this](FHttpRequestPtr HttpRequest, FHttpResponsePtr HttpResponse, bool bSucceeded)
+		{
+			CHECK(bSucceeded);
+			const int64 NoneCountAfter = HttpModule->GetHttpManager().GetHttpStats().RequestResults[static_cast<uint8>(EHttpFailureReason::None)].load();
+			CHECK(NoneCountAfter == 1);
+		});
+	}
+
+	SECTION("Failed request with ConnectionError increments RequestResults[ConnectionError]")
+	{
+		HttpModule->HttpConnectionTimeout = 0.1f;
+		CVarHttpUrlPatternsToMockResponse->Set(TEXT("www.epicgames.com->0"));
+
+		HttpRequest->SetURL(TEXT("https://www.epicgames.com"));
+		HttpRequest->OnProcessRequestComplete().BindLambda([this](FHttpRequestPtr HttpRequest, FHttpResponsePtr HttpResponse, bool bSucceeded)
+		{
+			CHECK(!bSucceeded);
+			CHECK(HttpRequest->GetFailureReason() == EHttpFailureReason::ConnectionError);
+			const int64 ConnectionErrorCountAfter = HttpModule->GetHttpManager().GetHttpStats().RequestResults[static_cast<uint8>(EHttpFailureReason::ConnectionError)].load();
+			CHECK(ConnectionErrorCountAfter == 1);
+			CVarHttpUrlPatternsToMockResponse->Set(TEXT(""));
+		});
+	}
+
+	HttpRequest->ProcessRequest();
+}
+
 #if UE_HTTP_SUPPORT_UNIX_SOCKET
 
 TEST_CASE_METHOD(FWaitUntilCompleteHttpFixture, "Http Methods over Unix Domain Socket", HTTP_TAG)
diff --git a/Engine/Source/Runtime/Online/HTTP/Private/GenericPlatform/HttpRequestCommon.cpp b/Engine/Source/Runtime/Online/HTTP/Private/GenericPlatform/HttpRequestCommon.cpp
--- a/Engine/Source/Runtime/Online/HTTP/Private/GenericPlatform/HttpRequestCommon.cpp
+++ b/Engine/Source/Runtime/Online/HTTP/Private/GenericPlatform/HttpRequestCommon.cpp
@@ -328,6 +328,7 @@ void FHttpRequestCommon::HandleRequestSucceed()
 	LogResponse(ResponseCommon);
 
 	FHttpModule::Get().GetHttpManager().RecordStatTimeToConnect(ConnectTime);
+	FHttpModule::Get().GetHttpManager().RecordStatRequestSucceeded();
 }
 
 void FHttpRequestCommon::HandleRequestFailed()
@@ -355,6 +356,8 @@ void FHttpRequestCommon::HandleRequestFailed()
 
 	SetStatus(EHttpRequestStatus::Failed);
 
+	FHttpModule::Get().GetHttpManager().RecordStatRequestFailed(FailureReason);
+
 	LogFailure();
 }
 
diff --git a/Engine/Source/Runtime/Online/HTTP/Private/HttpManager.cpp b/Engine/Source/Runtime/Online/HTTP/Private/HttpManager.cpp
--- a/Engine/Source/Runtime/Online/HTTP/Private/HttpManager.cpp
+++ b/Engine/Source/Runtime/Online/HTTP/Private/HttpManager.cpp
@@ -711,6 +711,20 @@ void FHttpManager::RecordMaxTimeToWaitInQueue(float Duration)
 	HttpStats.MaxTimeToWaitInQueue = FGenericPlatformMath::Max(Duration, HttpStats.MaxTimeToWaitInQueue.load());
 }
 
+void FHttpManager::RecordStatRequestSucceeded()
+{
+	++HttpStats.RequestResults[static_cast<uint8>(EHttpFailureReason::None)];
+}
+
+void FHttpManager::RecordStatRequestFailed(EHttpFailureReason Reason)
+{
+	const uint8 Index = static_cast<uint8>(Reason);
+	if (Index < static_cast<uint8>(EHttpFailureReason::Count))
+	{
+		++HttpStats.RequestResults[Index];
+	}
+}
+
 void FHttpManager::RecordPlatformStats(const FHttpStatsPlatform& PlatformStats)
 {
 	HttpStats.PlatformStats = PlatformStats;
diff --git a/Engine/Source/Runtime/Online/HTTP/Public/HttpManager.h b/Engine/Source/Runtime/Online/HTTP/Public/HttpManager.h
--- a/Engine/Source/Runtime/Online/HTTP/Public/HttpManager.h
+++ b/Engine/Source/Runtime/Online/HTTP/Public/HttpManager.h
@@ -89,6 +89,12 @@ struct FHttpStats
 	/** The optional http stats on specific platform */
 	TOptional<FHttpStatsPlatform> PlatformStats;
 
+	/**
+	 * Request completion counters indexed by EHttpFailureReason.
+	 * Index 0 (None) counts succeeded requests; indices 1+ count each failure reason.
+	 */
+	std::atomic<int64> RequestResults[static_cast<uint8>(EHttpFailureReason::Count)] = {};
+
 	FHttpStats()
 	{
 	}
@@ -105,6 +111,10 @@ struct FHttpStats
 		, HttpDurationMsAvg(Other.HttpDurationMsAvg.load())
 		, PlatformStats(Other.PlatformStats)
 	{
+		for (uint8 i = 0; i < static_cast<uint8>(EHttpFailureReason::Count); ++i)
+		{
+			RequestResults[i] = Other.RequestResults[i].load();
+		}
 	}
 };
 
@@ -444,6 +454,12 @@ PACKAGE_SCOPE:
 	/** Record the time to wait in queue, to have a general idea how long the client usually wait before actually starting, to adjust the requests */
 	HTTP_API void RecordMaxTimeToWaitInQueue(float Duration);
 
+	/** Record a succeeded request */
+	HTTP_API void RecordStatRequestSucceeded();
+
+	/** Record a failed request by its failure reason */
+	HTTP_API void RecordStatRequestFailed(EHttpFailureReason Reason);
+
 	/** Record platform specific stats */
 	HTTP_API void RecordPlatformStats(const FHttpStatsPlatform& PlatformStats);
 
diff --git a/Engine/Source/Runtime/Online/HTTP/Public/Interfaces/IHttpBase.h b/Engine/Source/Runtime/Online/HTTP/Public/Interfaces/IHttpBase.h
--- a/Engine/Source/Runtime/Online/HTTP/Public/Interfaces/IHttpBase.h
+++ b/Engine/Source/Runtime/Online/HTTP/Public/Interfaces/IHttpBase.h
@@ -67,7 +67,8 @@ enum class EHttpFailureReason : uint8
 	// When user's configured time is up, it's TimedOut no matter which step this Http request is in.
 	TimedOut,
 	Other,
-	ResponseTooLarge
+	ResponseTooLarge,
+	Count
 };
 
 /** @return the stringified version of the enum passed in */
