[Http]Add test case for validating the post body.

#Synced-CL 51774074

diff --git a/Engine/Source/Programs/WebTests/Private/TestHttp.cpp b/Engine/Source/Programs/WebTests/Private/TestHttp.cpp
--- a/Engine/Source/Programs/WebTests/Private/TestHttp.cpp
+++ b/Engine/Source/Programs/WebTests/Private/TestHttp.cpp
@@ -219,6 +219,7 @@ public:
 	const FString UrlMockLatency(uint32 Latency) const { return FString::Format(TEXT("{0}/mock_latency/{1}/"), { *UrlHttpTests(), Latency }); }
 	const FString UrlMockStatus(uint32 StatusCode) const { return FString::Format(TEXT("{0}/mock_status/{1}/"), { *UrlHttpTests(), StatusCode }); }
 	const FString UrlUnixSocketHttpTests() const { return "http://localhost/webtests/unixsockettests"; }
+	const FString UrlPostEchoBody() const { return FString::Format(TEXT("{0}/post_echo_body"), { *UrlHttpTests() }); }
 
 	FString WebServerIp;
 	FString WebServerUnixSocket;
@@ -3487,6 +3488,43 @@ TEST_CASE_METHOD(FWaitUntilCompleteHttpFixture, "RecordStatRequest increments Re
 	HttpRequest->ProcessRequest();
 }
 
+TEST_CASE_METHOD(FWaitUntilCompleteHttpFixture, "POST body is received correctly by the server", HTTP_TAG)
+{
+	struct FEchoBodyResponse : public FJsonSerializable
+	{
+		FString WorkspaceId;
+		FString UxTraceId;
+
+		BEGIN_JSON_SERIALIZER
+			JSON_SERIALIZE("workspaceId", WorkspaceId);
+			JSON_SERIALIZE("uxTraceId", UxTraceId);
+		END_JSON_SERIALIZER
+	};
+
+	const FString RequestBody = TEXT(R"({
+		"workspaceId":"default",
+		"uxTraceId":"64dd71404524bc1e868d8787ee2be934"
+	})");
+
+	TSharedRef<IHttpRequest> HttpRequest = CreateRequest();
+	HttpRequest->SetURL(UrlPostEchoBody());
+	HttpRequest->SetVerb(TEXT("POST"));
+	HttpRequest->SetHeader(TEXT("Content-Type"), TEXT("application/json"));
+	HttpRequest->SetContentAsString(RequestBody);
+	HttpRequest->OnProcessRequestComplete().BindLambda([](FHttpRequestPtr HttpRequest, FHttpResponsePtr HttpResponse, bool bSucceeded)
+	{
+		CHECK(bSucceeded);
+		REQUIRE(HttpResponse != nullptr);
+		CHECK(HttpResponse->GetResponseCode() == 200);
+
+		FEchoBodyResponse Response;
+		REQUIRE(Response.FromJsonStringView(HttpResponse->GetContentAsUtf8StringView()));
+		CHECK(Response.WorkspaceId == TEXT("default"));
+		CHECK(Response.UxTraceId == TEXT("64dd71404524bc1e868d8787ee2be934"));
+	});
+	HttpRequest->ProcessRequest();
+}
+
 #if UE_HTTP_SUPPORT_UNIX_SOCKET
 
 TEST_CASE_METHOD(FWaitUntilCompleteHttpFixture, "Http Methods over Unix Domain Socket", HTTP_TAG)
diff --git a/Engine/Source/Programs/WebTestsServer/httptests/urls.py b/Engine/Source/Programs/WebTestsServer/httptests/urls.py
--- a/Engine/Source/Programs/WebTestsServer/httptests/urls.py
+++ b/Engine/Source/Programs/WebTestsServer/httptests/urls.py
@@ -14,4 +14,5 @@ urlpatterns = [
     path("httptests/redirect_to", views.redirect_to, name="redirect_to"),
     path("httptests/mock_latency/<int:latency>/", views.mock_latency, name="mock_latency"),
     path("httptests/mock_status/<int:status_code>/", views.mock_status, name="mock_status"),
+    path("httptests/post_echo_body", views.post_echo_body, name="post_echo_body"),
 ]
diff --git a/Engine/Source/Programs/WebTestsServer/httptests/views.py b/Engine/Source/Programs/WebTestsServer/httptests/views.py
--- a/Engine/Source/Programs/WebTestsServer/httptests/views.py
+++ b/Engine/Source/Programs/WebTestsServer/httptests/views.py
@@ -93,6 +93,10 @@ def mock_latency(request, latency):
     time.sleep(latency)
     return JsonResponse({})
 
+@api_view(['POST'])
+def post_echo_body(request):
+    return JsonResponse(request.data)
+
 @api_view(['GET', 'POST', 'DELETE', 'PUT'])
 def mock_status(request, status_code):
     json_result = json.dumps({
