[Http]Don't call curl_easy_cleanup after shutting down libCurl.

Possible reason of the crash:

It seems possible to hold reference of http request by user code, or by request adapter from retry manager after http module shutdown. Destroying CurlHttpRequest and call curl_easy_cleanup at that point could cause undefined behavior. The callstack show HttpManager is there but could be the re-created http module and manager because FHttpModule::Get auto-created it(changed to not auto-create in changelist 52144947, so we can catch that as well).

#Synced-CL 52039227

diff --git a/Engine/Source/Programs/WebTests/Private/TestHttp.cpp b/Engine/Source/Programs/WebTests/Private/TestHttp.cpp
--- a/Engine/Source/Programs/WebTests/Private/TestHttp.cpp
+++ b/Engine/Source/Programs/WebTests/Private/TestHttp.cpp
@@ -3525,6 +3525,122 @@ TEST_CASE_METHOD(FWaitUntilCompleteHttpFixture, "POST body is received correctly
 	HttpRequest->ProcessRequest();
 }
 
+TEST_CASE_METHOD(FHttpModuleTestFixture, "RetryHttpRequestWithDelay called during Flush does not crash", HTTP_TAG)
+{
+	// Regression test for EXCEPTION_ACCESS_VIOLATION in FHttpRetrySystem::FManager::RetryHttpRequestWithDelay
+	// during engine shutdown (tick [0]). The crash occurs when a retry request's connection error completion
+	// fires during FHttpManager::Flush(Shutdown) and calls AddHttpThreadTask on the still-alive Thread.
+	//
+	// Setup: start a retry request against a non-routable IP so it fails with a connection error quickly.
+	// Then immediately call ShutdownModule() — the request completes during Flush, RetryHttpRequestWithDelay
+	// is called, and AddHttpThreadTask must not crash.
+	//
+	// Use near-zero backoff so that the queued event-loop timer fires almost immediately, preventing
+	// StopThread() from hanging while waiting for the event loop to drain.
+	if (!bRetryEnabled)
+	{
+		return;
+	}
+
+	DisableWarningsInThisTest();
+
+	// Short connection timeout so the request fails quickly during Flush.
+	// GetHttpConnectionTimeout() returns int32, so use at least 1 (0.5f truncates to 0 which check-fails).
+	HttpModule->HttpConnectionTimeout = 1.0f;
+
+	// Near-zero backoff: Compute(1) = Base^(1+1) * Coeff = 1.001^2 * ~0 ≈ 0
+	// This ensures StopThread() doesn't block waiting for a multi-second timer to drain.
+	FHttpRetrySystem::FExponentialBackoffCurve RetryBackoffCurve;
+	RetryBackoffCurve.Base = 1.001f;
+	RetryBackoffCurve.MinCoefficient = 0.0f;
+	RetryBackoffCurve.MaxCoefficient = 0.001f;
+
+	TSharedPtr<IHttpRequest> HttpRequest = HttpRetryManager->CreateRequest(
+		1/*InRetryLimitCountOverride*/,
+		FHttpRetrySystem::FRetryTimeoutRelativeSecondsSetting(),
+		{ EHttpResponseCodes::TooManyRequests, EHttpResponseCodes::ServiceUnavail },
+		FHttpRetrySystem::FRetryVerbs(),
+		FHttpRetrySystem::FRetryDomainsPtr(),
+		1/*InRetryLimitCountForConnectionErrorOverride*/,
+		RetryBackoffCurve
+	);
+
+	SetInvalidDomainToTestConnectTimeout();
+	HttpRequest->SetURL(UrlToTestMethods());
+	HttpRequest->SetVerb(TEXT("GET"));
+
+	HttpRequest->ProcessRequest();
+
+	// Shut down immediately — the request is still in-flight.
+	// During Flush(Shutdown), the connection times out, HttpOnProcessRequestComplete fires,
+	// RetryHttpRequestWithDelay calls AddHttpThreadTask. This must not crash.
+	ShutdownModule();
+	HttpRequest.Reset();
+}
+
+TEST_CASE_METHOD(FHttpModuleTestFixture, "Shutdown http module while request is in back-off delay after connection error", HTTP_TAG)
+{
+	// Regression test for shutdown during back-off after connection error.
+	// Verifies that FManager destruction while a delayed retry task is pending doesn't crash,
+	// particularly when CVarHttpCancelSupportDuringBackOffEnabled=true causes ClearInCaseOfRetry
+	// to be called before queuing the delayed lambda.
+	if (!bRetryEnabled)
+	{
+		return;
+	}
+
+	DisableWarningsInThisTest();
+
+	HttpModule->HttpConnectionTimeout = 1.0f;
+
+	FHttpRetrySystem::FExponentialBackoffCurve RetryBackoffCurve;
+	RetryBackoffCurve.MinCoefficient = 1.0f; // no jitter, default backoff will be ~4s
+
+	TSharedPtr<IHttpRequest> HttpRequest = HttpRetryManager->CreateRequest(
+		1/*InRetryLimitCountOverride*/,
+		FHttpRetrySystem::FRetryTimeoutRelativeSecondsSetting(),
+		{ EHttpResponseCodes::TooManyRequests, EHttpResponseCodes::ServiceUnavail },
+		FHttpRetrySystem::FRetryVerbs(),
+		FHttpRetrySystem::FRetryDomainsPtr(),
+		1/*InRetryLimitCountForConnectionErrorOverride*/,
+		RetryBackoffCurve
+	);
+
+	SetInvalidDomainToTestConnectTimeout();
+	HttpRequest->SetURL(UrlToTestMethods());
+	HttpRequest->SetVerb(TEXT("GET"));
+
+	std::atomic<bool> bWillRetryFired = false;
+	HttpRequest->OnRequestWillRetry().BindLambda([&bWillRetryFired](FHttpRequestPtr /*Request*/, FHttpResponsePtr /*Response*/, float /*LockoutPeriod*/) {
+		bWillRetryFired = true;
+		// Do NOT cancel — let the delayed task get queued, then shut down while it's pending
+	});
+
+	HttpRequest->ProcessRequest();
+
+	// Tick until WillRetry fires (meaning ClearInCaseOfRetry has been called and delayed task queued)
+	const double StartTime = FPlatformTime::Seconds();
+	double LastTick = StartTime;
+	while (!bWillRetryFired)
+	{
+		const double Now = FPlatformTime::Seconds();
+		HttpModule->GetHttpManager().Tick(Now - LastTick);
+		LastTick = Now;
+		FPlatformProcess::Sleep(1.0f / 60);
+		if (Now - StartTime > 10.0)
+		{
+			break;
+		}
+	}
+
+	CHECK(bWillRetryFired);
+
+	// Shut down while the delayed retry task is still pending (backoff ~4s not elapsed yet).
+	// HttpRequest is still alive here, so its destructor will fire after the platform HTTP library is shut down.
+	ShutdownModule();
+	HttpRequest.Reset(); // release our ref; destructor fires here with platform HTTP library already shut down
+}
+
 #if UE_HTTP_SUPPORT_UNIX_SOCKET
 
 TEST_CASE_METHOD(FWaitUntilCompleteHttpFixture, "Http Methods over Unix Domain Socket", HTTP_TAG)
diff --git a/Engine/Source/Runtime/Online/HTTP/Private/Curl/CurlHttp.cpp b/Engine/Source/Runtime/Online/HTTP/Private/Curl/CurlHttp.cpp
--- a/Engine/Source/Runtime/Online/HTTP/Private/Curl/CurlHttp.cpp
+++ b/Engine/Source/Runtime/Online/HTTP/Private/Curl/CurlHttp.cpp
@@ -246,18 +246,28 @@ bool FCurlHttpRequest::InitEasyHandle()
 
 void FCurlHttpRequest::DestroyEasyHandle()
 {
-	if (EasyHandle)
+	if (FCurlHttpManager::IsInit())
 	{
-		// cleanup the handle first (that order is used in howtos)
-		curl_easy_cleanup(EasyHandle);
-		EasyHandle = nullptr;
+		if (EasyHandle)
+		{
+			// cleanup the handle first (that order is used in howtos)
+			curl_easy_cleanup(EasyHandle);
+			EasyHandle = nullptr;
+		}
+ 
+		// destroy headers list
+		if (HeaderList)
+		{
+			curl_slist_free_all(HeaderList);
+			HeaderList = nullptr;
+		}
 	}
-
-	// destroy headers list
-	if (HeaderList)
+	else
 	{
-		curl_slist_free_all(HeaderList);
-		HeaderList = nullptr;
+		// Curl library has already been shut down (e.g. user code holds a reference to the request shared ptr).
+		// Curl handles cannot be safely cleaned up; memory is leaked. Flush all HTTP requests before shutting
+		// down the HTTP module to avoid this.
+		UE_LOG(LogHttp, Warning, TEXT("Curl library has already been shut down. Curl handles will be leaked. Flush all HTTP requests before shutting down the HTTP module to avoid this."));
 	}
 }
 
