[WebSockets] Remove FHttpModule::Get() dependency(Still depends on HTTP for static methods for now).

#Synced-CL 52039227

#Synced-CL 52039227

diff --git a/Engine/Source/Programs/WebTests/Private/TestHttp.cpp b/Engine/Source/Programs/WebTests/Private/TestHttp.cpp
--- a/Engine/Source/Programs/WebTests/Private/TestHttp.cpp
+++ b/Engine/Source/Programs/WebTests/Private/TestHttp.cpp
@@ -24,6 +24,7 @@
 #include "Serialization/JsonSerializerMacros.h"
 #include "catch2/generators/catch_generators.hpp"
 #include "Ssl.h"
+#include "WebTestFixtures.h"
 
 /**
  *  HTTP Tests
@@ -45,7 +46,6 @@
  // For xCurl, this is the minimum chunk size to trigger data received callback immediately, without waiting additional time for the recv buffer to fill up
 #define HTTP_TEST_TIMEOUT_CHUNK_SIZE 32*1024
 
-extern TAutoConsoleVariable<bool> CVarHttpInsecureProtocolEnabled;
 extern TAutoConsoleVariable<int32> CVarHttpMaxConcurrentRequests;
 extern TAutoConsoleVariable<int32> CVarHttpEventLoopEnableChance;
 extern TAutoConsoleVariable<FString> CVarHttpUrlPatternsToMockResponse;
@@ -143,13 +143,14 @@ class FHttpModuleTestFixture
 {
 public:
 	FHttpModuleTestFixture()
-		: WebServerIp(TEXT("127.0.0.1"))
-		, WebServerHttpPort(8000)
-		, bRunHeavyTests(false)
+		: bRunHeavyTests(false)
 		, bRetryEnabled(true)
 		, bEventLoopEnabled(true)
 	{
-		ParseSettingsFromCommandLine();
+		FParse::Bool(FCommandLine::Get(), TEXT("run_heavy_tests="), bRunHeavyTests);
+		FParse::Bool(FCommandLine::Get(), TEXT("retry_enabled="), bRetryEnabled);
+		FParse::Bool(FCommandLine::Get(), TEXT("eventloop_enabled="), bEventLoopEnabled);
+		FParse::Value(FCommandLine::Get(), TEXT("web_server_unix_socket="), WebServerUnixSocket);
 
 		CVarHttpEventLoopEnableChance->Set(bEventLoopEnabled ? 100 : 0);
 
@@ -185,15 +186,6 @@ public:
 		ShutdownModule();
 	}
 
-	void ParseSettingsFromCommandLine()
-	{
-		FParse::Value(FCommandLine::Get(), TEXT("web_server_ip="), WebServerIp);
-		FParse::Bool(FCommandLine::Get(), TEXT("run_heavy_tests="), bRunHeavyTests);
-		FParse::Bool(FCommandLine::Get(), TEXT("retry_enabled="), bRetryEnabled);
-		FParse::Bool(FCommandLine::Get(), TEXT("eventloop_enabled="), bEventLoopEnabled);
-		FParse::Value(FCommandLine::Get(), TEXT("web_server_unix_socket="), WebServerUnixSocket);
-	}
-
 	void DisableWarningsInThisTest()
 	{
 		HttpTestLogLevelInitializer.DisableWarningsInThisTest();
@@ -209,10 +201,9 @@ public:
 		return bRetryEnabled ? HttpRetryManager->CreateRequest() : HttpModule->CreateRequest();
 	}
 
-	void SetInvalidDomainToTestConnectTimeout() { WebServerIp = TEXT("10.255.255.1"); WebServerHttpPort = 8765; } // non-routable IP address with a random port
-	const FString UrlDomain() const { return FString::Format(TEXT("{0}:{1}"), { *WebServerIp, WebServerHttpPort }); }
-	const FString UrlBase() const { return FString::Format(TEXT("http://{0}"), { *UrlDomain() }); }
-	const FString UrlHttpTests() const { return FString::Format(TEXT("{0}/webtests/httptests"), { *UrlBase() }); }
+	void SetInvalidDomainToTestConnectTimeout() { ServerSettings.WebServerIp = TEXT("10.255.255.1"); ServerSettings.WebServerPort = 8765; } // non-routable IP address with a random port
+	const FString UrlDomain() const { return FString::Format(TEXT("{0}:{1}"), { *ServerSettings.WebServerIp, ServerSettings.WebServerPort }); }
+	const FString UrlHttpTests() const { return ServerSettings.UrlHttpTests(); }
 	const FString UrlToTestMethods() const { return FString::Format(TEXT("{0}/methods"), { *UrlHttpTests() }); }
 	const FString UrlStreamDownload(uint32 Chunks, uint32 ChunkSize, uint32 ChunkLatency=0) { return FString::Format(TEXT("{0}/streaming_download/{1}/{2}/{3}/"), { *UrlHttpTests(), Chunks, ChunkSize, ChunkLatency }); }
 	const FString UrlStreamUpload() { return FString::Format(TEXT("{0}/streaming_upload_put"), { *UrlHttpTests() }); }
@@ -221,9 +212,8 @@ public:
 	const FString UrlUnixSocketHttpTests() const { return "http://localhost/webtests/unixsockettests"; }
 	const FString UrlPostEchoBody() const { return FString::Format(TEXT("{0}/post_echo_body"), { *UrlHttpTests() }); }
 
-	FString WebServerIp;
+	FWebTestServerSettings ServerSettings;
 	FString WebServerUnixSocket;
-	uint32 WebServerHttpPort;
 	FMockHttpModule* HttpModule = nullptr;
 	bool bRunHeavyTests;
 	bool bRetryEnabled;
diff --git a/Engine/Source/Programs/WebTests/Private/TestWebSockets.cpp b/Engine/Source/Programs/WebTests/Private/TestWebSockets.cpp
--- a/Engine/Source/Programs/WebTests/Private/TestWebSockets.cpp
+++ b/Engine/Source/Programs/WebTests/Private/TestWebSockets.cpp
@@ -1,15 +1,15 @@
 // Copyright Epic Games, Inc. All Rights Reserved.
-#include "Containers/BackgroundableTicker.h"
 #include "CoreMinimal.h"
+#include "Containers/BackgroundableTicker.h"
 #include "HAL/IConsoleManager.h"
 #include "HAL/PlatformProcess.h"
 #include "Http.h"
-#include "HttpManager.h"
 #include "Misc/CommandLine.h"
 #include "WebSocketsLog.h"
 #include "WebSocketsModule.h"
 #include "IWebSocket.h"
 #include "TestHarness.h"
+#include "WebTestFixtures.h"
 
 /**
  *  WebSockets Tests
@@ -25,25 +25,14 @@
 
 #define WEBSOCKETS_TAG "[WebSockets]"
 
-extern TAutoConsoleVariable<bool> CVarHttpInsecureProtocolEnabled;
-
 class FWebSocketsModuleTestFixture
 {
 public:
 	FWebSocketsModuleTestFixture()
-		: WebServerIp(TEXT("127.0.0.1"))
-		, WebServerWebSocketsPort(8000)
-		, OldVerbosity(LogWebSockets.GetVerbosity())
+		: OldVerbosity(LogWebSockets.GetVerbosity())
 	{
 		CVarHttpInsecureProtocolEnabled->Set(true);
 
-		ParseSettingsFromCommandLine();
-
-		// Init HTTP module because websockets module has dependency on it when get proxy
-		HttpModule = new FHttpModule();
-		IModuleInterface* HttpModuleInterface = HttpModule;
-		HttpModuleInterface->StartupModule();
-
 		WebSocketsModule = new FWebSocketsModule();
 		IModuleInterface* Module = WebSocketsModule;
 		Module->StartupModule();
@@ -51,16 +40,12 @@ public:
 
 	virtual ~FWebSocketsModuleTestFixture()
 	{
-		IModuleInterface* WebSocketsModuleInterface = WebSocketsModule;
-		WebSocketsModuleInterface->ShutdownModule();
-		delete WebSocketsModuleInterface;
-
-		// Simulate the flush tick when shutdown, to make sure tasks added by FHttpModule::Get().GetHttpManager().AddGameThreadTask won't crash
-		HttpModule->GetHttpManager().Tick(0.0);
-
-		IModuleInterface* HttpModuleInterface = HttpModule;
-		HttpModuleInterface->ShutdownModule();
-		delete HttpModuleInterface;
+		if (WebSocketsModule != nullptr)
+		{
+			IModuleInterface* WebSocketsModuleInterface = WebSocketsModule;
+			WebSocketsModuleInterface->ShutdownModule();
+			delete WebSocketsModuleInterface;
+		}
 
 		if (OldVerbosity != LogWebSockets.GetVerbosity())
 		{
@@ -68,25 +53,16 @@ public:
 		}
 	}
 
-	void ParseSettingsFromCommandLine()
-	{
-		FParse::Value(FCommandLine::Get(), TEXT("web_server_ip="), WebServerIp);
-		FParse::Value(FCommandLine::Get(), TEXT("web_server_websockets_port="), WebServerWebSocketsPort);
-	}
-
 	void DisableWarningsInThisTest()
 	{
 		LogWebSockets.SetVerbosity(ELogVerbosity::Error);
 	}
 
-	const FString UrlWithInvalidPortToTestConnectTimeout() const { return FString::Format(TEXT("ws://{0}:{1}"), { *WebServerIp, 8765 }); }
-	const FString UrlBase() const { return FString::Format(TEXT("ws://{0}:{1}"), { *WebServerIp, WebServerWebSocketsPort }); }
-	const FString UrlWebSocketsTests() const { return FString::Format(TEXT("{0}/webtests/websocketstests"), { *UrlBase() }); }
+	const FString UrlWithInvalidPortToTestConnectTimeout() const { return FString::Format(TEXT("ws://{0}:{1}"), { *ServerSettings.WebServerIp, 8765 }); }
+	const FString UrlWebSocketsTests() const { return ServerSettings.UrlWebSocketsTests(); }
 
-	FString WebServerIp;
-	uint32 WebServerWebSocketsPort;
+	FWebTestServerSettings ServerSettings;
 	FWebSocketsModule* WebSocketsModule;
-	FHttpModule* HttpModule;
 	ELogVerbosity::Type OldVerbosity;
 };
 
@@ -99,28 +75,14 @@ public:
 
 	~FRunUntilQuitRequestedFixture()
 	{
-		RunUntilQuitRequested();
+		TickRunner.RunUntilQuitRequested();
 	}
 
-	void SimulateEngineTick()
-	{
-		FTSBackgroundableTicker::GetCoreTicker().Tick(TickFrequency);
-		FTSTicker::GetCoreTicker().Tick(TickFrequency);
-		FPlatformProcess::Sleep(TickFrequency);
-		HttpModule->GetHttpManager().Tick(TickFrequency);
-	}
+	void RequestQuit()          { TickRunner.bQuitRequested = true; }
+	void SetSucceeded()         { TickRunner.bSucceeded = true; }
+	bool HasSucceeded() const   { return TickRunner.bSucceeded; }
 
-	void RunUntilQuitRequested()
-	{
-		while (!bQuitRequested)
-		{
-			SimulateEngineTick();
-		}
-	}
-
-	float TickFrequency = 1.0f / 60; /*60 FPS*/;
-	bool bQuitRequested = false;
-	bool bSucceed = false;
+	FWebTestTickRunner TickRunner;
 };
 
 TEST_CASE_METHOD(FRunUntilQuitRequestedFixture, "WebSockets can connect then send and receive message", WEBSOCKETS_TAG)
@@ -134,17 +96,17 @@ TEST_CASE_METHOD(FRunUntilQuitRequestedFixture, "WebSockets can connect then sen
 	WebSocket->OnMessage().AddLambda([this, WebSocket](const FString& MessageString) {
 		CHECK(MessageString == TEXT("hi websockets tests"));
 		WebSocket->Close();
-		bSucceed = true;
+		SetSucceeded();
 	});
 
 	WebSocket->OnClosed().AddLambda([this, WebSocket](int32 /* StatusCode */, const FString& /* Reason */, bool /* bWasClean */){
-		CHECK(bSucceed);
-		bQuitRequested = true;
+		CHECK(HasSucceeded());
+		RequestQuit();
 	});
 
 	WebSocket->OnConnectionError().AddLambda([this, WebSocket](const FString& /* Error */){
 		CHECK(false);
-		bQuitRequested = true;
+		RequestQuit();
 	});
 
 	WebSocket->Connect();
@@ -160,12 +122,12 @@ TEST_CASE_METHOD(FRunUntilQuitRequestedFixture, "WebSockets can be closed by ser
 
 	WebSocket->OnClosed().AddLambda([this, WebSocket](int32 StatusCode, const FString& /* Reason */, bool /* bWasClean */){
 		CHECK(StatusCode == 1000); // 1000 means default normal closure
-		bQuitRequested = true;
+		RequestQuit();
 	});
 
 	WebSocket->OnConnectionError().AddLambda([this, WebSocket](const FString& /* Error */){
 		CHECK(false);
-		bQuitRequested = true;
+		RequestQuit();
 	});
 
 	WebSocket->Connect();
@@ -178,12 +140,12 @@ TEST_CASE_METHOD(FRunUntilQuitRequestedFixture, "WebSockets module can shut down
 	TSharedPtr<IWebSocket> WebSocket = WebSocketsModule->CreateWebSocket(FString::Format(TEXT("{0}/echo/"), { *UrlWebSocketsTests() }));
 
 	WebSocket->OnConnected().AddLambda([this, WebSocket](){
-		bQuitRequested = true;
+		RequestQuit();
 	});
 
 	WebSocket->OnConnectionError().AddLambda([this, WebSocket](const FString& /* Error */){
 		CHECK(false);
-		bQuitRequested = true;
+		RequestQuit();
 	});
 
 	WebSocket->Connect();
@@ -200,7 +162,7 @@ public:
 	~FConnectWhenShutdownFixture()
 	{
 		// Waiting until closed, then the Connect call can actually launch without early return
-		RunUntilQuitRequested();
+		TickRunner.RunUntilQuitRequested();
 
 		WebSocket->Connect();
 	}
@@ -219,17 +181,201 @@ TEST_CASE_METHOD(FConnectWhenShutdownFixture, "WebSockets can call connect when
 	WebSocket->OnMessage().AddLambda([this](const FString& MessageString) {
 		CHECK(MessageString == TEXT("hi websockets tests"));
 		WebSocket->Close();
-		bSucceed = true;
+		SetSucceeded();
 	});
 
 	WebSocket->OnClosed().AddLambda([this](int32 /* StatusCode */, const FString& /* Reason */, bool /* bWasClean */){
-		CHECK(bSucceed);
-		bQuitRequested = true;
+		CHECK(HasSucceeded());
+		RequestQuit();
+	});
+
+	WebSocket->OnConnectionError().AddLambda([this](const FString& /* Error */){
+		CHECK(false);
+		RequestQuit();
+	});
+
+	WebSocket->Connect();
+}
+
+// Fixture that starts FWebSocketsModule first, then FHttpModule after, verifying the two can coexist
+// without double-initializing or double-finalizing the underlying SSL library.
+class FHttpAndWebSocketsFixture : public FRunUntilQuitRequestedFixture
+{
+public:
+	FHttpAndWebSocketsFixture()
+	{
+		HttpModule = new FHttpModule();
+		static_cast<IModuleInterface*>(HttpModule)->StartupModule();
+	}
+
+	~FHttpAndWebSocketsFixture()
+	{
+		// Run the tick loop here (before shutting down HttpModule) so all in-flight requests finish.
+		// bQuitRequested may already be true if the loop ran earlier; RunUntilQuitRequested exits immediately in that case.
+		TickRunner.RunUntilQuitRequested();
+
+		if (HttpModule != nullptr)
+		{
+			static_cast<IModuleInterface*>(HttpModule)->ShutdownModule();
+			delete HttpModule;
+			HttpModule = nullptr;
+		}
+	}
+
+	FHttpModule* HttpModule = nullptr;
+};
+
+TEST_CASE_METHOD(FHttpAndWebSocketsFixture, "WebSockets works correctly when HttpModule is also loaded", WEBSOCKETS_TAG)
+{
+	// Both a WebSocket echo and an HTTP GET must complete before the test passes.
+	auto bHttpDone = MakeShared<bool>(false);
+	auto bWsDone = MakeShared<bool>(false);
+
+	auto TryFinish = [this, bHttpDone, bWsDone]()
+	{
+		if (*bHttpDone && *bWsDone)
+		{
+			CHECK(HasSucceeded());
+			RequestQuit();
+		}
+	};
+
+	// Fire an HTTP request via FHttpModule while WebSockets is also active.
+	TSharedRef<IHttpRequest, ESPMode::ThreadSafe> HttpRequest = HttpModule->CreateRequest();
+	HttpRequest->SetURL(FString::Format(TEXT("{0}/methods"), { *ServerSettings.UrlHttpTests() }));
+	HttpRequest->SetVerb(TEXT("GET"));
+	HttpRequest->OnProcessRequestComplete().BindLambda([bHttpDone, TryFinish](FHttpRequestPtr /* Request */, FHttpResponsePtr /* Response */, bool bProcessedSuccessfully)
+	{
+		CHECK(bProcessedSuccessfully);
+		*bHttpDone = true;
+		TryFinish();
+	});
+	HttpRequest->ProcessRequest();
+
+	TSharedPtr<IWebSocket> WebSocket = WebSocketsModule->CreateWebSocket(FString::Format(TEXT("{0}/echo/"), { *UrlWebSocketsTests() }));
+
+	WebSocket->OnConnected().AddLambda([this, WebSocket](){
+		WebSocket->Send(TEXT("hi http+websockets coexistence test"));
+	});
+
+	WebSocket->OnMessage().AddLambda([this, WebSocket](const FString& MessageString) {
+		CHECK(MessageString == TEXT("hi http+websockets coexistence test"));
+		WebSocket->Close();
+		SetSucceeded();
+	});
+
+	WebSocket->OnClosed().AddLambda([bWsDone, TryFinish](int32 /* StatusCode */, const FString& /* Reason */, bool /* bWasClean */){
+		*bWsDone = true;
+		TryFinish();
+	});
+
+	WebSocket->OnConnectionError().AddLambda([this](const FString& /* Error */){
+		CHECK(false);
+		RequestQuit();
+	});
+
+	WebSocket->Connect();
+}
+
+// Fixture that starts FHttpModule first, then FWebSocketsModule after, verifying the two can coexist
+// regardless of initialization order.
+class FHttpFirstThenWebSocketsFixture
+{
+public:
+	FHttpFirstThenWebSocketsFixture()
+		: OldVerbosity(LogWebSockets.GetVerbosity())
+	{
+		CVarHttpInsecureProtocolEnabled->Set(true);
+
+		// Http starts first — this is the key difference from FHttpAndWebSocketsFixture.
+		HttpModule = new FHttpModule();
+		static_cast<IModuleInterface*>(HttpModule)->StartupModule();
+
+		// WebSockets starts second.
+		WebSocketsModule = new FWebSocketsModule();
+		static_cast<IModuleInterface*>(WebSocketsModule)->StartupModule();
+	}
+
+	~FHttpFirstThenWebSocketsFixture()
+	{
+		TickRunner.RunUntilQuitRequested();
+
+		// WebSockets shuts down first.
+		static_cast<IModuleInterface*>(WebSocketsModule)->ShutdownModule();
+		delete WebSocketsModule;
+		WebSocketsModule = nullptr;
+
+		// Http shuts down last.
+		static_cast<IModuleInterface*>(HttpModule)->ShutdownModule();
+		delete HttpModule;
+		HttpModule = nullptr;
+
+		if (OldVerbosity != LogWebSockets.GetVerbosity())
+		{
+			LogWebSockets.SetVerbosity(OldVerbosity);
+		}
+	}
+
+	const FString UrlWebSocketsTests() const { return ServerSettings.UrlWebSocketsTests(); }
+
+	void RequestQuit()          { TickRunner.bQuitRequested = true; }
+	void SetSucceeded()         { TickRunner.bSucceeded = true; }
+	bool HasSucceeded() const   { return TickRunner.bSucceeded; }
+
+	FWebTestTickRunner TickRunner;
+	FWebTestServerSettings ServerSettings;
+	FHttpModule* HttpModule = nullptr;
+	FWebSocketsModule* WebSocketsModule = nullptr;
+	ELogVerbosity::Type OldVerbosity;
+};
+
+TEST_CASE_METHOD(FHttpFirstThenWebSocketsFixture, "WebSockets works correctly when HttpModule is loaded first", "Exclude")
+{
+	// Both a WebSocket echo and an HTTP GET must complete before the test passes.
+	auto bHttpDone = MakeShared<bool>(false);
+	auto bWsDone = MakeShared<bool>(false);
+
+	auto TryFinish = [this, bHttpDone, bWsDone]()
+	{
+		if (*bHttpDone && *bWsDone)
+		{
+			CHECK(HasSucceeded());
+			RequestQuit();
+		}
+	};
+
+	// Fire an HTTP request via FHttpModule while WebSockets is also active.
+	TSharedRef<IHttpRequest, ESPMode::ThreadSafe> HttpRequest = HttpModule->CreateRequest();
+	HttpRequest->SetURL(FString::Format(TEXT("{0}/methods"), { *ServerSettings.UrlHttpTests() }));
+	HttpRequest->SetVerb(TEXT("GET"));
+	HttpRequest->OnProcessRequestComplete().BindLambda([bHttpDone, TryFinish](FHttpRequestPtr /* Request */, FHttpResponsePtr /* Response */, bool bProcessedSuccessfully)
+	{
+		CHECK(bProcessedSuccessfully);
+		*bHttpDone = true;
+		TryFinish();
+	});
+	HttpRequest->ProcessRequest();
+
+	TSharedPtr<IWebSocket> WebSocket = WebSocketsModule->CreateWebSocket(FString::Format(TEXT("{0}/echo/"), { *UrlWebSocketsTests() }));
+
+	WebSocket->OnConnected().AddLambda([this, WebSocket](){
+		WebSocket->Send(TEXT("hi http-first websockets coexistence test"));
+	});
+
+	WebSocket->OnMessage().AddLambda([this, WebSocket](const FString& MessageString) {
+		CHECK(MessageString == TEXT("hi http-first websockets coexistence test"));
+		WebSocket->Close();
+		SetSucceeded();
+	});
+
+	WebSocket->OnClosed().AddLambda([bWsDone, TryFinish](int32 /* StatusCode */, const FString& /* Reason */, bool /* bWasClean */){
+		*bWsDone = true;
+		TryFinish();
 	});
 
 	WebSocket->OnConnectionError().AddLambda([this](const FString& /* Error */){
 		CHECK(false);
-		bQuitRequested = true;
+		RequestQuit();
 	});
 
 	WebSocket->Connect();
diff --git a/Engine/Source/Programs/WebTests/Private/WebTestFixtures.h b/Engine/Source/Programs/WebTests/Private/WebTestFixtures.h
new file mode 100644
--- /dev/null
+++ b/Engine/Source/Programs/WebTests/Private/WebTestFixtures.h
@@ -0,0 +1,94 @@
+// Copyright Epic Games, Inc. All Rights Reserved.
+
+#pragma once
+
+#include "CoreMinimal.h"
+#include "Containers/BackgroundableTicker.h"
+#include "HAL/IConsoleManager.h"
+#include "HAL/PlatformProcess.h"
+#include "Misc/CommandLine.h"
+#include "Misc/Parse.h"
+
+/**
+ *  Shared fixture utilities for WebTests (Http and WebSockets tests)
+ *  -----------------------------------------------------------------------
+ *
+ *  Contains:
+ *  - CVarHttpInsecureProtocolEnabled extern (defined in HttpModule.cpp)
+ *  - FWebTestServerSettings: command-line driven server address helpers
+ *  - FWebTestTickRunner: engine tick simulation with quit-requested loop
+ *
+ *  -----------------------------------------------------------------------
+ */
+
+/** Defined in HttpModule.cpp — allows tests to connect over plain http:// / ws:// */
+extern TAutoConsoleVariable<bool> CVarHttpInsecureProtocolEnabled;
+
+/**
+ * Parses web test server connection settings from the command line.
+ * Both Http and WebSocket tests target the same web server process.
+ *
+ * Command-line parameters:
+ *   -web_server_ip=<ip>                    (default: 127.0.0.1)
+ *   -web_server_websockets_port=<port>      (default: 8000)
+ *   -web_server_http_port=<port>            (default: same as websockets port)
+ */
+struct FWebTestServerSettings
+{
+	FWebTestServerSettings()
+		: WebServerIp(TEXT("127.0.0.1"))
+		, WebServerPort(8000)
+	{
+		FParse::Value(FCommandLine::Get(), TEXT("web_server_ip="), WebServerIp);
+		// Accept either param name; websockets tests historically used web_server_websockets_port.
+		if (!FParse::Value(FCommandLine::Get(), TEXT("web_server_http_port="), WebServerPort))
+		{
+			FParse::Value(FCommandLine::Get(), TEXT("web_server_websockets_port="), WebServerPort);
+		}
+	}
+
+	/** Base HTTP URL: "http://<ip>:<port>" */
+	FString HttpUrlBase() const { return FString::Format(TEXT("http://{0}:{1}"), { *WebServerIp, WebServerPort }); }
+
+	/** URL prefix for the http test routes: "http://<ip>:<port>/webtests/httptests" */
+	FString UrlHttpTests() const { return FString::Format(TEXT("{0}/webtests/httptests"), { *HttpUrlBase() }); }
+
+	/** Base WebSocket URL: "ws://<ip>:<port>" */
+	FString WsUrlBase() const { return FString::Format(TEXT("ws://{0}:{1}"), { *WebServerIp, WebServerPort }); }
+
+	/** URL prefix for the websocket test routes: "ws://<ip>:<port>/webtests/websocketstests" */
+	FString UrlWebSocketsTests() const { return FString::Format(TEXT("{0}/webtests/websocketstests"), { *WsUrlBase() }); }
+
+	FString WebServerIp;
+	uint32  WebServerPort;
+};
+
+/**
+ * Helper that provides engine-tick simulation for integration tests that drive async operations
+ * (Http requests, WebSocket connections, etc.) from the game thread.
+ *
+ * Drives FTSBackgroundableTicker and FTSTicker at 60 FPS until RequestQuit() is called.
+ * Intended to be used as a member variable in test fixtures, not as a base class.
+ */
+class FWebTestTickRunner
+{
+public:
+	void SimulateEngineTick()
+	{
+		FTSBackgroundableTicker::GetCoreTicker().Tick(TickFrequency);
+		FTSTicker::GetCoreTicker().Tick(TickFrequency);
+		FPlatformProcess::Sleep(TickFrequency);
+	}
+
+	void RunUntilQuitRequested()
+	{
+		while (!bQuitRequested)
+		{
+			SimulateEngineTick();
+		}
+	}
+
+	float TickFrequency  = 1.0f / 60; /*60 FPS*/
+	bool  bQuitRequested = false;
+	bool  bSucceeded     = false;
+};
diff --git a/Engine/Source/Runtime/Online/HTTP/Private/GenericPlatform/GenericPlatformHttp.cpp b/Engine/Source/Runtime/Online/HTTP/Private/GenericPlatform/GenericPlatformHttp.cpp
--- a/Engine/Source/Runtime/Online/HTTP/Private/GenericPlatform/GenericPlatformHttp.cpp
+++ b/Engine/Source/Runtime/Online/HTTP/Private/GenericPlatform/GenericPlatformHttp.cpp
@@ -7,6 +7,9 @@
 #include "GenericPlatform/HttpRequestImpl.h"
 #include "HAL/CriticalSection.h"
 #include "HAL/IConsoleManager.h"
+#include "Misc/CommandLine.h"
+#include "Misc/ConfigCacheIni.h"
+#include "Misc/Parse.h"
 #include "Misc/Paths.h"
 #include "Misc/EngineVersion.h"
 #include "Misc/App.h"
@@ -620,6 +623,22 @@ FString FGenericPlatformHttp::EscapeUserAgentString(const FString& UnescapedStri
 	}
 }
 
+FString FGenericPlatformHttp::GetConfiguredProxyAddress()
+{
+	FString ProxyAddress;
+	if (!FParse::Value(FCommandLine::Get(), TEXT("httpproxy="), ProxyAddress))
+	{
+		if (!GConfig->GetString(TEXT("HTTP"), TEXT("HttpProxyAddress"), ProxyAddress, GEngineIni))
+		{
+			if (TOptional<FString> OperatingSystemProxyAddress = GetOperatingSystemProxyAddress())
+			{
+				ProxyAddress = MoveTemp(OperatingSystemProxyAddress.GetValue());
+			}
+		}
+	}
+	return ProxyAddress;
+}
+
 TOptional<FString> FGenericPlatformHttp::GetOperatingSystemProxyAddress()
 {
 	return TOptional<FString>();
diff --git a/Engine/Source/Runtime/Online/HTTP/Private/HttpManager.cpp b/Engine/Source/Runtime/Online/HTTP/Private/HttpManager.cpp
--- a/Engine/Source/Runtime/Online/HTTP/Private/HttpManager.cpp
+++ b/Engine/Source/Runtime/Online/HTTP/Private/HttpManager.cpp
@@ -287,30 +287,7 @@ bool FHttpManager::IsDomainAllowed(const FString& Url) const
 	HttpManagerInternal::CheckNoDuplicatedProtocol(Url);
 #endif // !UE_BUILD_SHIPPING
 
-	if (!URLRequestFilter.IsEmpty())
-	{
-		return URLRequestFilter.IsRequestAllowed(Url);
-	}
-
-#if !UE_BUILD_SHIPPING
-#if !(UE_GAME || UE_SERVER)
-	// Allowed domain filtering is opt-in in non-shipping non-game/server builds
-	static const bool bForceUseAllowList = FParse::Param(FCommandLine::Get(), TEXT("EnableHttpDomainRestrictions"));
-	if (!bForceUseAllowList)
-	{
-		return true;
-	}
-#else
-	// The check is on by default but allow non-shipping game/server builds to disable the filtering
-	static const bool bIgnoreAllowList = FParse::Param(FCommandLine::Get(), TEXT("DisableHttpDomainRestrictions"));
-	if (bIgnoreAllowList)
-	{
-		return true;
-	}
-#endif
-#endif // !UE_BUILD_SHIPPING
-
-	return true;
+	return URLRequestFilter.IsRequestAllowed(Url);
 }
 
 /*static*/
diff --git a/Engine/Source/Runtime/Online/HTTP/Private/HttpModule.cpp b/Engine/Source/Runtime/Online/HTTP/Private/HttpModule.cpp
--- a/Engine/Source/Runtime/Online/HTTP/Private/HttpModule.cpp
+++ b/Engine/Source/Runtime/Online/HTTP/Private/HttpModule.cpp
@@ -112,16 +112,7 @@ void FHttpModule::StartupModule()
 	FCoreDelegates::TSOnConfigSectionsChanged().AddRaw(this, &FHttpModule::OnConfigSectionsChanged);
 	UpdateConfigs();
 
-	if (!FParse::Value(FCommandLine::Get(), TEXT("httpproxy="), ProxyAddress))
-	{
-		if (!GConfig->GetString(TEXT("HTTP"), TEXT("HttpProxyAddress"), ProxyAddress, GEngineIni))
-		{
-			if (TOptional<FString> OperatingSystemProxyAddress = FPlatformHttp::GetOperatingSystemProxyAddress())
-			{
-				ProxyAddress = MoveTemp(OperatingSystemProxyAddress.GetValue());
-			}
-		}
-	}
+	ProxyAddress = FPlatformHttp::GetConfiguredProxyAddress();
 
 	// Load from a configurable array of modules at this point, so things that need to bind to the SDK Manager init hooks can do so.
 	TArray<FString> ModulesToLoad;
diff --git a/Engine/Source/Runtime/Online/HTTP/Private/WinHttp/Support/WinHttpSession.cpp b/Engine/Source/Runtime/Online/HTTP/Private/WinHttp/Support/WinHttpSession.cpp
--- a/Engine/Source/Runtime/Online/HTTP/Private/WinHttp/Support/WinHttpSession.cpp
+++ b/Engine/Source/Runtime/Online/HTTP/Private/WinHttp/Support/WinHttpSession.cpp
@@ -7,7 +7,7 @@
 #include "WinHttp/Support/WinHttpConnection.h"
 #include "WinHttp/Support/WinHttpTypes.h"
 #include "GenericPlatform/GenericPlatformHttp.h"
-#include "Http.h"
+#include "Misc/ConfigCacheIni.h"
 
 #include "Windows/AllowWindowsPlatformTypes.h"
 #include <errhandlingapi.h>
@@ -82,10 +82,13 @@ FWinHttpSession::FWinHttpSession(uint32 SecurityProtocolFlags, const bool bInFor
 		UE_LOG(LogWinHttp, Verbose, TEXT("WinHttp local machine does not support for compression"));
 	}
 
-	FHttpModule& HttpModule = FHttpModule::Get();
-	const FTimespan ConnectionTimeout = FTimespan::FromSeconds(HttpModule.GetHttpConnectionTimeout() > 0 ? HttpModule.GetHttpConnectionTimeout() : 0);
+	float HttpConnectionTimeout = 30.0f;
+	float HttpActivityTimeout = 30.0f;
+	GConfig->GetFloat(TEXT("HTTP"), TEXT("HttpConnectionTimeout"), HttpConnectionTimeout, GEngineIni);
+	GConfig->GetFloat(TEXT("HTTP"), TEXT("HttpActivityTimeout"), HttpActivityTimeout, GEngineIni);
+	const FTimespan ConnectionTimeout = FTimespan::FromSeconds(HttpConnectionTimeout > 0 ? HttpConnectionTimeout : 0);
 	const FTimespan ResolveTimeout = ConnectionTimeout;
-	const FTimespan ReceiveTimeout = FTimespan::FromSeconds(HttpModule.GetHttpActivityTimeout() > 0 ? HttpModule.GetHttpActivityTimeout() : 0);
+	const FTimespan ReceiveTimeout = FTimespan::FromSeconds(HttpActivityTimeout > 0 ? HttpActivityTimeout : 0);
 	const FTimespan SendTimeout = ReceiveTimeout;
 
 	if (!WinHttpSetTimeouts(SessionHandle.Get(), ResolveTimeout.GetTotalMilliseconds(), ConnectionTimeout.GetTotalMilliseconds(), SendTimeout.GetTotalMilliseconds(), ReceiveTimeout.GetTotalMilliseconds()))
diff --git a/Engine/Source/Runtime/Online/HTTP/Public/GenericPlatform/GenericPlatformHttp.h b/Engine/Source/Runtime/Online/HTTP/Public/GenericPlatform/GenericPlatformHttp.h
--- a/Engine/Source/Runtime/Online/HTTP/Public/GenericPlatform/GenericPlatformHttp.h
+++ b/Engine/Source/Runtime/Online/HTTP/Public/GenericPlatform/GenericPlatformHttp.h
@@ -182,6 +182,16 @@ public:
 	 * @return the escaped string.
 	 */
 
+	/**
+	 * Get the configured proxy address, using the same precedence as the HTTP module:
+	 * 1. Command-line: -httpproxy=
+	 * 2. INI: [HTTP] HttpProxyAddress
+	 * 3. Operating system proxy
+	 *
+	 * @return the proxy address, or an empty string if none is configured
+	 */
+	static HTTP_API FString GetConfiguredProxyAddress();
+
 	/**
 	 * Get the proxy address specified by the operating system
 	 *
diff --git a/Engine/Source/Runtime/Online/WebSockets/Private/Lws/LwsWebSocket.cpp b/Engine/Source/Runtime/Online/WebSockets/Private/Lws/LwsWebSocket.cpp
--- a/Engine/Source/Runtime/Online/WebSockets/Private/Lws/LwsWebSocket.cpp
+++ b/Engine/Source/Runtime/Online/WebSockets/Private/Lws/LwsWebSocket.cpp
@@ -13,9 +13,8 @@
 #include "Misc/EmbeddedCommunication.h"
 #include "Misc/ScopeLock.h"
 #include "Misc/ConfigCacheIni.h"
-#include "HttpModule.h"
-#include "HttpManager.h"
-#include "PlatformHttp.h"
+#include "GenericPlatform/GenericPlatformHttp.h"
+#include "Misc/URLRequestFilter.h"
 
 // FLwsSendBuffer 
 FLwsSendBuffer::FLwsSendBuffer(const uint8* Data, SIZE_T Size, bool bInIsBinary)
@@ -112,8 +111,8 @@ void FLwsWebSocket::Connect()
 
 	if (!FLwsWebSocketsManager::Get().bDisableDomainAllowlist)
 	{
-		FHttpManager& HttpManager = FHttpModule::Get().GetHttpManager();
-		if (!HttpManager.IsDomainAllowed(Url))
+		UE::Core::FURLRequestFilter DomainFilter(TEXT("Online.HttpManager"), GEngineIni);
+		if (!DomainFilter.IsRequestAllowed(Url))
 		{
 			State = EState::Error;
 			UE_LOG(LogWebSockets, Warning, TEXT("FLwsWebSocket[%d]::Connect: %s is not allowlisted. Refusing to connect."), Identifier, *Url);
diff --git a/Engine/Source/Runtime/Online/WebSockets/Private/Lws/LwsWebSocketsManager.cpp b/Engine/Source/Runtime/Online/WebSockets/Private/Lws/LwsWebSocketsManager.cpp
--- a/Engine/Source/Runtime/Online/WebSockets/Private/Lws/LwsWebSocketsManager.cpp
+++ b/Engine/Source/Runtime/Online/WebSockets/Private/Lws/LwsWebSocketsManager.cpp
@@ -9,7 +9,7 @@
 #if WITH_SSL
 #include "Ssl.h"
 #endif
-#include "HttpModule.h"
+#include "PlatformHttp.h"
 #include "WebSocketsLog.h"
 #include "Containers/BackgroundableTicker.h"
 #include "HAL/PlatformTime.h"
@@ -193,8 +193,8 @@ void FLwsWebSocketsManager::InitWebSockets(TArrayView<const FString> Protocols)
 	GConfig->GetInt(TEXT("WebSockets.LibWebSockets"), TEXT("PingPongInterval"), PingPongInterval, GEngineIni);
 	ContextInfo.ws_ping_pong_interval = PingPongInterval;
 
-	// HTTP proxy
-	const FString& ProxyAddress = FHttpModule::Get().GetProxyAddress();
+	// HTTP proxy — read using same precedence as FHttpModule: command-line, INI, OS proxy
+	FString ProxyAddress = FPlatformHttp::GetConfiguredProxyAddress();
 	TOptional<FTCHARToUTF8> Converter;
 	if (!ProxyAddress.IsEmpty())
 	{
diff --git a/Engine/Source/Runtime/Online/WebSockets/Private/WinHttp/WinHttpWebSocket.cpp b/Engine/Source/Runtime/Online/WebSockets/Private/WinHttp/WinHttpWebSocket.cpp
--- a/Engine/Source/Runtime/Online/WebSockets/Private/WinHttp/WinHttpWebSocket.cpp
+++ b/Engine/Source/Runtime/Online/WebSockets/Private/WinHttp/WinHttpWebSocket.cpp
@@ -10,9 +10,7 @@
 #include "WinHttp/Support/WinHttpSession.h"
 
 #include "WebSocketsLog.h"
-#include "HttpManager.h"
-#include "HttpModule.h"
-#include "GenericPlatform/GenericPlatformHttp.h"
+#include "Misc/URLRequestFilter.h"
 #include "Misc/ConfigCacheIni.h"
 
 const TCHAR* LexToString(const EWebSocketConnectionState State)
@@ -74,8 +72,8 @@ void FWinHttpWebSocket::Connect()
 	GConfig->GetBool(TEXT("WinHttpWebSocket"), TEXT("bDisableDomainAllowlist"), bDisableDomainAllowlist, GEngineIni);
 	if (!bDisableDomainAllowlist)
 	{
-		FHttpManager& HttpManager = FHttpModule::Get().GetHttpManager();
-		if (!HttpManager.IsDomainAllowed(Url))
+		UE::Core::FURLRequestFilter DomainFilter(TEXT("Online.HttpManager"), GEngineIni);
+		if (!DomainFilter.IsRequestAllowed(Url))
 		{
 			UE_LOG(LogWebSockets, Warning, TEXT("WinHttp WebSocket[%p]: %s is not in the allowed list, refusing to connect."), this, *Url);
 			HandleCloseComplete(EWebSocketConnectionState::FailedToConnect, UE_WEBSOCKET_CLOSE_APP_FAILURE, FString("Invalid Domain"));
diff --git a/Engine/Source/Runtime/Online/WebSockets/Private/WinHttp/WinHttpWebSocketsManager.cpp b/Engine/Source/Runtime/Online/WebSockets/Private/WinHttp/WinHttpWebSocketsManager.cpp
--- a/Engine/Source/Runtime/Online/WebSockets/Private/WinHttp/WinHttpWebSocketsManager.cpp
+++ b/Engine/Source/Runtime/Online/WebSockets/Private/WinHttp/WinHttpWebSocketsManager.cpp
@@ -5,9 +5,8 @@
 #include "WinHttp/WinHttpWebSocketsManager.h"
 #include "WinHttp/WinHttpWebSocket.h"
 #include "WinHttp/WinHttpHttpManager.h"
-#include "Http.h"
+#include "PlatformHttp.h"
 #include "Modules/ModuleManager.h"
-#include "HttpModule.h"
 #include "Containers/BackgroundableTicker.h"
 #include "Stats/Stats.h"
 
